Made GameDebugger module not include Engine.h #UE4

diff --git a/Engine/Source/Developer/GameplayDebugger/Private/GameplayDebuggerExec.cpp b/Engine/Source/Developer/GameplayDebugger/Private/GameplayDebuggerExec.cpp
--- a/Engine/Source/Developer/GameplayDebugger/Private/GameplayDebuggerExec.cpp
+++ b/Engine/Source/Developer/GameplayDebugger/Private/GameplayDebuggerExec.cpp
@@ -5,6 +5,12 @@
 #if !(UE_BUILD_SHIPPING || UE_BUILD_TEST)
 #include "Misc/CoreMisc.h"
 
+#if WITH_EDITOR
+#include "Editor.h"
+#include "EditorViewportClient.h"
+#include "LevelEditorViewport.h"
+#endif //WITH_EDITOR
+
 struct FGameplayDebuggerExec : public FSelfRegisteringExec
 {
 	FGameplayDebuggerExec()
@@ -25,7 +31,6 @@ FGameplayDebuggerExec GameplayDebuggerExecInstance;
 
 TWeakObjectPtr<AGameplayDebuggingReplicator> FGameplayDebuggerExec::GetDebuggingReplicator(UWorld* InWorld)
 {
-#if !(UE_BUILD_SHIPPING || UE_BUILD_TEST)
 	if (CachedDebuggingReplicator.IsValid() && CachedDebuggingReplicator->GetWorld() == InWorld)
 	{
 		return CachedDebuggingReplicator;
@@ -40,7 +45,7 @@ TWeakObjectPtr<AGameplayDebuggingReplicator> FGameplayDebuggerExec::GetDebugging
 			return CachedDebuggingReplicator;
 		}
 	}
-#endif
+
 	return NULL;
 }
 
@@ -48,7 +53,6 @@ bool FGameplayDebuggerExec::Exec(UWorld* Inworld, const TCHAR* Cmd, FOutputDevic
 {
 	bool bHandled = false;
 
-#if !(UE_BUILD_SHIPPING || UE_BUILD_TEST)
 	APlayerController* PC = Inworld ? Inworld->GetFirstPlayerController() : NULL;
 	if (!Inworld || !PC)
 	{
@@ -102,7 +106,6 @@ bool FGameplayDebuggerExec::Exec(UWorld* Inworld, const TCHAR* Cmd, FOutputDevic
 		}
 		bHandled = true;
 	}
-#endif
 
 	return bHandled;
 }
diff --git a/Engine/Source/Developer/GameplayDebugger/Private/GameplayDebuggerPrivate.h b/Engine/Source/Developer/GameplayDebugger/Private/GameplayDebuggerPrivate.h
--- a/Engine/Source/Developer/GameplayDebugger/Private/GameplayDebuggerPrivate.h
+++ b/Engine/Source/Developer/GameplayDebugger/Private/GameplayDebuggerPrivate.h
@@ -2,12 +2,32 @@
 
 #pragma once
 
-#if WITH_EDITOR
-#	include "UnrealEd.h"
-#endif
-
 #include "Core.h"
-#include "Engine.h"
+#include "CoreUObject.h"
+#include "InputCore.h"
+#include "EngineDefines.h"
+#include "EngineSettings.h"
+#include "EngineStats.h"
+#include "EngineLogs.h"
+#include "EngineGlobals.h"
+
+#include "Engine/EngineBaseTypes.h"
+#include "Engine/EngineTypes.h"
+#include "Components/ActorComponent.h"
+#include "GameFramework/Actor.h"
+#include "CollisionQueryParams.h"
+#include "WorldCollision.h"
+#include "Components/PrimitiveComponent.h"
+#include "GameFramework/PlayerController.h"
+#include "AI/Navigation/RecastNavMesh.h"
+#include "GameFramework/Pawn.h"
+#include "Components/CapsuleComponent.h"
+#include "Engine/Engine.h"
+
+#include "Engine/Canvas.h"
+#include "EngineUtils.h"
+
+//////////////////////////////////////////////////////////////////////////
 
 #include "GameplayDebuggingTypes.h"
 #include "GameplayDebugger.h"
diff --git a/Engine/Source/Developer/GameplayDebugger/Private/GameplayDebuggingComponent.cpp b/Engine/Source/Developer/GameplayDebugger/Private/GameplayDebuggingComponent.cpp
--- a/Engine/Source/Developer/GameplayDebugger/Private/GameplayDebuggingComponent.cpp
+++ b/Engine/Source/Developer/GameplayDebugger/Private/GameplayDebuggingComponent.cpp
@@ -12,6 +12,7 @@
 #include "EnvironmentQuery/EnvQueryTypes.h"
 #include "BehaviorTreeDelegates.h"
 #include "Navigation/NavigationComponent.h"
+#include "GameFramework/PlayerState.h"
 
 #include "GameplayDebuggingComponent.h"
 
diff --git a/Engine/Source/Developer/GameplayDebugger/Private/GameplayDebuggingControllerComponent.cpp b/Engine/Source/Developer/GameplayDebugger/Private/GameplayDebuggingControllerComponent.cpp
--- a/Engine/Source/Developer/GameplayDebugger/Private/GameplayDebuggingControllerComponent.cpp
+++ b/Engine/Source/Developer/GameplayDebugger/Private/GameplayDebuggingControllerComponent.cpp
@@ -2,7 +2,9 @@
 #include "GameplayDebuggerPrivate.h"
 #include "Net/UnrealNetwork.h"
 #include "GameFramework/HUD.h"
+#include "GameFramework/PlayerInput.h"
 #include "BehaviorTreeDelegates.h"
+#include "TimerManager.h"
 #include "VisualLog.h"
 
 #if WITH_EDITOR
diff --git a/Engine/Source/Developer/GameplayDebugger/Private/GameplayDebuggingHUDComponent.cpp b/Engine/Source/Developer/GameplayDebugger/Private/GameplayDebuggingHUDComponent.cpp
--- a/Engine/Source/Developer/GameplayDebugger/Private/GameplayDebuggingHUDComponent.cpp
+++ b/Engine/Source/Developer/GameplayDebugger/Private/GameplayDebuggingHUDComponent.cpp
@@ -8,7 +8,11 @@
 #include "GameplayDebuggingComponent.h"
 #include "GameplayDebuggingHUDComponent.h"
 #include "GameplayDebuggingControllerComponent.h"
+#include "CanvasItem.h"
+#include "AI/Navigation/NavigationSystem.h"
 #include "Regex.h"
+#include "DrawDebugHelpers.h"
+#include "TimerManager.h"					// Game play timers
 
 DEFINE_LOG_CATEGORY_STATIC(LogHUD, Log, All);
 
diff --git a/Engine/Source/Developer/GameplayDebugger/Private/GameplayDebuggingReplicator.cpp b/Engine/Source/Developer/GameplayDebugger/Private/GameplayDebuggingReplicator.cpp
--- a/Engine/Source/Developer/GameplayDebugger/Private/GameplayDebuggingReplicator.cpp
+++ b/Engine/Source/Developer/GameplayDebugger/Private/GameplayDebuggingReplicator.cpp
@@ -4,6 +4,10 @@
 #include "GameplayDebuggingComponent.h"
 #include "GameplayDebuggingHUDComponent.h"
 #include "GameplayDebuggingReplicator.h"
+#include "BehaviorTreeDelegates.h"
+#if WITH_EDITOR
+#include "Editor/EditorEngine.h"
+#endif // WITH_EDITOR
 #include "UnrealNetwork.h"
 
 AGameplayDebuggingReplicator::AGameplayDebuggingReplicator(const class FPostConstructInitializeProperties& PCIP)
diff --git a/Engine/Source/Editor/UnrealEd/Classes/Editor/EditorEngine.h b/Engine/Source/Editor/UnrealEd/Classes/Editor/EditorEngine.h
--- a/Engine/Source/Editor/UnrealEd/Classes/Editor/EditorEngine.h
+++ b/Engine/Source/Editor/UnrealEd/Classes/Editor/EditorEngine.h
@@ -8,6 +8,10 @@
 #include "../Settings/LevelEditorViewportSettings.h"
 #include "EditorEngine.generated.h"
 
+class FAssetData;
+class UAnimSequence;
+class USkeleton;
+
 //
 // Things to set in mapSetBrush.
 //
diff --git a/Engine/Source/Editor/UnrealEd/Public/Editor.h b/Engine/Source/Editor/UnrealEd/Public/Editor.h
--- a/Engine/Source/Editor/UnrealEd/Public/Editor.h
+++ b/Engine/Source/Editor/UnrealEd/Public/Editor.h
@@ -11,6 +11,8 @@
 #include "Commandlets/EditorCommandlets.h"
 #include "EditorUndoClient.h"
 #include "TickableEditorObject.h"
+#include "Editor/UnrealEdTypes.h"
+#include "Editor/EditorEngine.h"
 
 #define CAMERA_ZOOM_DAMPEN			200.f
 
@@ -20,6 +22,7 @@ class FGeomEdge;
 class FGeomPoly;
 class FGeomObject;
 class FScopedTransaction;
+class UFactory;
 struct FBlueprintGraphActionListBuilder;
 
 
diff --git a/Engine/Source/Editor/UnrealEd/Public/EditorModeTools.h b/Engine/Source/Editor/UnrealEd/Public/EditorModeTools.h
--- a/Engine/Source/Editor/UnrealEd/Public/EditorModeTools.h
+++ b/Engine/Source/Editor/UnrealEd/Public/EditorModeTools.h
@@ -2,6 +2,8 @@
 
 #pragma once
 
+class FViewport;
+class FLevelEditorViewportClient;
 
 enum EGeomSelectionType
 {
diff --git a/Engine/Source/Editor/UnrealEd/Public/EditorViewportClient.h b/Engine/Source/Editor/UnrealEd/Public/EditorViewportClient.h
--- a/Engine/Source/Editor/UnrealEd/Public/EditorViewportClient.h
+++ b/Engine/Source/Editor/UnrealEd/Public/EditorViewportClient.h
@@ -2,6 +2,12 @@
 
 #pragma once
 
+#include "EditorViewportClient.h"
+#include "UnrealWidget.h"
+#include "Framework/Commands/Commands.h"
+#include "Layout/SlateRect.h"
+#include "EditorComponents.h"
+
 class FEditorCameraController;
 class FCameraControllerUserImpulseData;
 class FMouseDeltaTracker;
diff --git a/Engine/Source/Editor/UnrealEd/Public/LevelEditorViewport.h b/Engine/Source/Editor/UnrealEd/Public/LevelEditorViewport.h
--- a/Engine/Source/Editor/UnrealEd/Public/LevelEditorViewport.h
+++ b/Engine/Source/Editor/UnrealEd/Public/LevelEditorViewport.h
@@ -3,6 +3,8 @@
 
 #pragma once
 
+#include "EditorViewportClient.h"
+
 // Forward declarations.
 
 class IMatineeBase;
diff --git a/Engine/Source/Editor/UnrealEd/Public/UnrealWidget.h b/Engine/Source/Editor/UnrealEd/Public/UnrealWidget.h
--- a/Engine/Source/Editor/UnrealEd/Public/UnrealWidget.h
+++ b/Engine/Source/Editor/UnrealEd/Public/UnrealWidget.h
@@ -6,6 +6,8 @@
 
 #pragma once
 
+class FEditorViewportClient;
+
 class FWidget : public FGCObject
 {
 public:
diff --git a/Engine/Source/Runtime/Engine/Classes/Debug/DebugDrawService.h b/Engine/Source/Runtime/Engine/Classes/Debug/DebugDrawService.h
--- a/Engine/Source/Runtime/Engine/Classes/Debug/DebugDrawService.h
+++ b/Engine/Source/Runtime/Engine/Classes/Debug/DebugDrawService.h
@@ -2,6 +2,7 @@
 #pragma once
 
 #include "ShowFlags.h"
+#include "Kismet/BlueprintFunctionLibrary.h"
 #include "DebugDrawService.generated.h"
 
 class FSceneView;
diff --git a/Engine/Source/Runtime/Engine/Classes/Engine/NetConnection.h b/Engine/Source/Runtime/Engine/Classes/Engine/NetConnection.h
--- a/Engine/Source/Runtime/Engine/Classes/Engine/NetConnection.h
+++ b/Engine/Source/Runtime/Engine/Classes/Engine/NetConnection.h
@@ -7,6 +7,8 @@
 #pragma once
 #include "Net/DataBunch.h"
 #include "Engine/Channel.h"
+#include "Engine/Player.h"
+#include "Engine/NetDriver.h"
 #include "NetConnection.generated.h"
 
 class FObjectReplicator;
diff --git a/Engine/Source/Runtime/Engine/Classes/GameFramework/Actor.h b/Engine/Source/Runtime/Engine/Classes/GameFramework/Actor.h
--- a/Engine/Source/Runtime/Engine/Classes/GameFramework/Actor.h
+++ b/Engine/Source/Runtime/Engine/Classes/GameFramework/Actor.h
@@ -9,6 +9,8 @@
 #include "Components/SceneComponent.h"
 #include "ComponentInstanceDataCache.h"
 
+class FTimerManager; 
+
 UENUM(BlueprintType)
 namespace EEndPlayReason
 {
diff --git a/Engine/Source/Runtime/Engine/Private/AI/Navigation/RecastNavMeshGenerator.h b/Engine/Source/Runtime/Engine/Private/AI/Navigation/RecastNavMeshGenerator.h
--- a/Engine/Source/Runtime/Engine/Private/AI/Navigation/RecastNavMeshGenerator.h
+++ b/Engine/Source/Runtime/Engine/Private/AI/Navigation/RecastNavMeshGenerator.h
@@ -10,6 +10,9 @@
 
 #define MAX_VERTS_PER_POLY	6
 
+struct FNavigationRelevantData;
+class FNavigationOctree;
+
 struct FRecastBuildConfig : public rcConfig
 {
 	/** controls whether voxel filterring will be applied (via FRecastTileGenerator::ApplyVoxelFilter) */
diff --git a/Engine/Source/Runtime/Engine/Private/NavMeshRenderingHelpers.h b/Engine/Source/Runtime/Engine/Private/NavMeshRenderingHelpers.h
--- a/Engine/Source/Runtime/Engine/Private/NavMeshRenderingHelpers.h
+++ b/Engine/Source/Runtime/Engine/Private/NavMeshRenderingHelpers.h
@@ -3,6 +3,9 @@
 #include "DebugRenderSceneProxy.h"
 #include "DynamicMeshBuilder.h"
 #include "AI/Navigation/RecastHelpers.h"
+#include "BatchedElements.h"
+#include "AI/Navigation/NavMeshRenderingComponent.h"
+#include "Materials/Material.h"
 
 static const FColor NavMeshRenderColor_Recast_TriangleEdges(255,255,255);
 static const FColor NavMeshRenderColor_Recast_TileEdges(16,16,16,32);
diff --git a/Engine/Source/Runtime/Engine/Public/CanvasItem.h b/Engine/Source/Runtime/Engine/Public/CanvasItem.h
--- a/Engine/Source/Runtime/Engine/Public/CanvasItem.h
+++ b/Engine/Source/Runtime/Engine/Public/CanvasItem.h
@@ -6,6 +6,8 @@
 
 #pragma once
 
+#include "BatchedElements.h"
+
 class FCanvasItem
 {
 public:
diff --git a/Engine/Source/Runtime/Engine/Public/Net/DataChannel.h b/Engine/Source/Runtime/Engine/Public/Net/DataChannel.h
--- a/Engine/Source/Runtime/Engine/Public/Net/DataChannel.h
+++ b/Engine/Source/Runtime/Engine/Public/Net/DataChannel.h
@@ -6,6 +6,8 @@
 
 #pragma once
 #include "Engine/Channel.h"
+#include "Engine/NetConnection.h"
+#include "Engine/ChildConnection.h"
 
 /*-----------------------------------------------------------------------------
 	UControlChannel base class.
diff --git a/Engine/Source/Runtime/Engine/Public/WorldCollision.h b/Engine/Source/Runtime/Engine/Public/WorldCollision.h
--- a/Engine/Source/Runtime/Engine/Public/WorldCollision.h
+++ b/Engine/Source/Runtime/Engine/Public/WorldCollision.h
@@ -8,7 +8,7 @@
 
 #include "UniquePtr.h"
 #include "CollisionQueryParams.h"
-#include "Engine/World.h"
+//#include "Engine/World.h"
 
 /**
  * Async Trace Data Structures 
diff --git a/Engine/Source/Runtime/Slate/Public/Framework/Commands/Commands.h b/Engine/Source/Runtime/Slate/Public/Framework/Commands/Commands.h
--- a/Engine/Source/Runtime/Slate/Public/Framework/Commands/Commands.h
+++ b/Engine/Source/Runtime/Slate/Public/Framework/Commands/Commands.h
@@ -2,6 +2,8 @@
 
 #pragma once
 
+#include "Framework/Commands/UICommandInfo.h"
+
 //
 // COMMAND DECLARATION
 //
diff --git a/Engine/Source/Runtime/Slate/Public/Framework/Commands/UICommandInfo.h b/Engine/Source/Runtime/Slate/Public/Framework/Commands/UICommandInfo.h
--- a/Engine/Source/Runtime/Slate/Public/Framework/Commands/UICommandInfo.h
+++ b/Engine/Source/Runtime/Slate/Public/Framework/Commands/UICommandInfo.h
@@ -6,6 +6,10 @@
 
 #pragma once
 
+#include "Layout/Visibility.h"
+#include "Misc/Attribute.h"
+#include "Framework/Commands/InputGesture.h"
+#include "Textures/SlateIcon.h"
 
 /** Types of user interfaces that can be associated with a user interface action */
 namespace EUserInterfaceActionType
diff --git a/Engine/Source/Runtime/Slate/Public/Framework/MultiBox/SToolBarComboButtonBlock.h b/Engine/Source/Runtime/Slate/Public/Framework/MultiBox/SToolBarComboButtonBlock.h
--- a/Engine/Source/Runtime/Slate/Public/Framework/MultiBox/SToolBarComboButtonBlock.h
+++ b/Engine/Source/Runtime/Slate/Public/Framework/MultiBox/SToolBarComboButtonBlock.h
@@ -2,6 +2,8 @@
 
 #pragma once
 
+struct FSlateIcon;
+
 /**
  * Tool bar combo button MultiBlock
  */
diff --git a/Engine/Source/Runtime/SlateCore/Public/Textures/SlateIcon.h b/Engine/Source/Runtime/SlateCore/Public/Textures/SlateIcon.h
--- a/Engine/Source/Runtime/SlateCore/Public/Textures/SlateIcon.h
+++ b/Engine/Source/Runtime/SlateCore/Public/Textures/SlateIcon.h
@@ -6,6 +6,7 @@
 
 #pragma once
 
+struct FSlateBrush;
 
 /**
  * Struct used to represent an icon in Slate
