Fixed TTP #338489 "[Editor] GameplayDebuggingReplicator visible in new Worlds"

diff --git a/Engine/Source/Developer/GameplayDebugger/Private/GameplayDebuggingHUDComponent.cpp b/Engine/Source/Developer/GameplayDebugger/Private/GameplayDebuggingHUDComponent.cpp
--- a/Engine/Source/Developer/GameplayDebugger/Private/GameplayDebuggingHUDComponent.cpp
+++ b/Engine/Source/Developer/GameplayDebugger/Private/GameplayDebuggingHUDComponent.cpp
@@ -22,6 +22,12 @@ AGameplayDebuggingHUDComponent::AGameplayDebuggingHUDComponent(const class FPost
 	, EngineShowFlags(EShowFlagInitMode::ESFIM_Game)
 {
 	World = NULL;
+	SetIsTemporarilyHiddenInEditor(true);
+	SetActorHiddenInGame(false);
+	bHiddenEdLevel = true;
+	bHiddenEdLayer = true;
+	bHiddenEd = true;
+	bEditable = false;
 }
 
 void AGameplayDebuggingHUDComponent::Render()
diff --git a/Engine/Source/Developer/GameplayDebugger/Private/GameplayDebuggingReplicator.cpp b/Engine/Source/Developer/GameplayDebugger/Private/GameplayDebuggingReplicator.cpp
--- a/Engine/Source/Developer/GameplayDebugger/Private/GameplayDebuggingReplicator.cpp
+++ b/Engine/Source/Developer/GameplayDebugger/Private/GameplayDebuggingReplicator.cpp
@@ -18,12 +18,12 @@ AGameplayDebuggingReplicator::AGameplayDebuggingReplicator(const class FPostCons
 	RootComponent = SceneComponent;
 
 	SetTickableWhenPaused(true);
-	//SetIsTemporarilyHiddenInEditor(true);
-	//SetActorHiddenInGame(true);
-	//bHiddenEdLevel = true;
-	//bHiddenEdLayer = true;
-	//bHiddenEd = true;
-	//bEditable = false;
+	SetIsTemporarilyHiddenInEditor(true);
+	SetActorHiddenInGame(false);
+	bHiddenEdLevel = true;
+	bHiddenEdLayer = true;
+	bHiddenEd = true;
+	bEditable = false;
 }
 
 void AGameplayDebuggingReplicator::GetLifetimeReplicatedProps(TArray< FLifetimeProperty > & OutLifetimeProps) const
