Fixed compiler error for WITH_EDITOR_ONLY_DATA.

diff --git a/Engine/Source/Developer/GameplayDebugger/Private/GameplayDebuggingHUDComponent.cpp b/Engine/Source/Developer/GameplayDebugger/Private/GameplayDebuggingHUDComponent.cpp
--- a/Engine/Source/Developer/GameplayDebugger/Private/GameplayDebuggingHUDComponent.cpp
+++ b/Engine/Source/Developer/GameplayDebugger/Private/GameplayDebuggingHUDComponent.cpp
@@ -22,12 +22,14 @@ AGameplayDebuggingHUDComponent::AGameplayDebuggingHUDComponent(const class FPost
 	, EngineShowFlags(EShowFlagInitMode::ESFIM_Game)
 {
 	World = NULL;
+#if WITH_EDITORONLY_DATA
 	SetIsTemporarilyHiddenInEditor(true);
 	SetActorHiddenInGame(false);
 	bHiddenEdLevel = true;
 	bHiddenEdLayer = true;
 	bHiddenEd = true;
 	bEditable = false;
+#endif
 }
 
 void AGameplayDebuggingHUDComponent::Render()
diff --git a/Engine/Source/Developer/GameplayDebugger/Private/GameplayDebuggingReplicator.cpp b/Engine/Source/Developer/GameplayDebugger/Private/GameplayDebuggingReplicator.cpp
--- a/Engine/Source/Developer/GameplayDebugger/Private/GameplayDebuggingReplicator.cpp
+++ b/Engine/Source/Developer/GameplayDebugger/Private/GameplayDebuggingReplicator.cpp
@@ -16,7 +16,7 @@ AGameplayDebuggingReplicator::AGameplayDebuggingReplicator(const class FPostCons
 {
 	TSubobjectPtr<USceneComponent> SceneComponent = PCIP.CreateDefaultSubobject<USceneComponent>(this, TEXT("SceneComponent"));
 	RootComponent = SceneComponent;
-
+#if WITH_EDITORONLY_DATA
 	SetTickableWhenPaused(true);
 	SetIsTemporarilyHiddenInEditor(true);
 	SetActorHiddenInGame(false);
@@ -24,6 +24,7 @@ AGameplayDebuggingReplicator::AGameplayDebuggingReplicator(const class FPostCons
 	bHiddenEdLayer = true;
 	bHiddenEd = true;
 	bEditable = false;
+#endif
 }
 
 void AGameplayDebuggingReplicator::GetLifetimeReplicatedProps(TArray< FLifetimeProperty > & OutLifetimeProps) const
