Remove SharedPCH settings from AIModule and GameplayDebugger. It doesn't make sense to create a shared PCH from a private header file (because you can't/shouldn't share private headers), and it breaks compiling any project under Rocket where private headers aren't distributed.

diff --git a/Engine/Source/Developer/GameplayDebugger/GameplayDebugger.Build.cs b/Engine/Source/Developer/GameplayDebugger/GameplayDebugger.Build.cs
--- a/Engine/Source/Developer/GameplayDebugger/GameplayDebugger.Build.cs
+++ b/Engine/Source/Developer/GameplayDebugger/GameplayDebugger.Build.cs
@@ -6,8 +6,6 @@ namespace UnrealBuildTool.Rules
     {
         public GameplayDebugger(TargetInfo Target)
         {
-            SharedPCHHeaderFile = "Developer/GameplayDebugger/Private/GameplayDebuggerPrivate.h";
-
             PublicIncludePaths.AddRange(
                 new string[] {
 				    "Developer/GameplayDebugger/Public",
diff --git a/Engine/Source/Runtime/AIModule/AIModule.Build.cs b/Engine/Source/Runtime/AIModule/AIModule.Build.cs
--- a/Engine/Source/Runtime/AIModule/AIModule.Build.cs
+++ b/Engine/Source/Runtime/AIModule/AIModule.Build.cs
@@ -6,8 +6,6 @@ namespace UnrealBuildTool.Rules
 	{
         public AIModule(TargetInfo Target)
 		{
-            SharedPCHHeaderFile = "Runtime/AIModule/Private/AIModulePrivate.h";
-
 			PublicIncludePaths.AddRange(
 				new string[] {
                     "Runtime/AIModule/Public",
