Added correct multi-client support for GameplayDebugger. #UE4
Removed some header dependencies from GameplayDebugger (GameplayDebugger's pch no longer including editor headers, etc.) #UE4

diff --git a/Engine/Source/Developer/GameplayDebugger/Classes/GameplayDebuggingComponent.h b/Engine/Source/Developer/GameplayDebugger/Classes/GameplayDebuggingComponent.h
--- a/Engine/Source/Developer/GameplayDebugger/Classes/GameplayDebuggingComponent.h
+++ b/Engine/Source/Developer/GameplayDebugger/Classes/GameplayDebuggingComponent.h
@@ -209,7 +209,7 @@ public:
 	virtual void DestroyRenderState_Concurrent() override;
 
 protected:
-	void SelectTargetToDebug(FDebugContext& Context);
+	void SelectTargetToDebug();
 
 	//APlayerController* PlayerOwner;
 
diff --git a/Engine/Source/Developer/GameplayDebugger/Classes/GameplayDebuggingControllerComponent.h b/Engine/Source/Developer/GameplayDebugger/Classes/GameplayDebuggingControllerComponent.h
--- a/Engine/Source/Developer/GameplayDebugger/Classes/GameplayDebuggingControllerComponent.h
+++ b/Engine/Source/Developer/GameplayDebugger/Classes/GameplayDebuggingControllerComponent.h
@@ -69,7 +69,7 @@ protected:
 	virtual void ToggleAIDebugView_SetView9();
 
 	virtual void BindAIDebugViewKeys();
-	AGameplayDebuggingReplicator* GetDebuggingReplicator();
+	AGameplayDebuggingReplicator* GetDebuggingReplicator() const;
 
 	TWeakObjectPtr<APlayerController> PlayerOwner;
 
diff --git a/Engine/Source/Developer/GameplayDebugger/Classes/GameplayDebuggingHUDComponent.h b/Engine/Source/Developer/GameplayDebugger/Classes/GameplayDebuggingHUDComponent.h
--- a/Engine/Source/Developer/GameplayDebugger/Classes/GameplayDebuggingHUDComponent.h
+++ b/Engine/Source/Developer/GameplayDebugger/Classes/GameplayDebuggingHUDComponent.h
@@ -1,7 +1,6 @@
 // Copyright 1998-2014 Epic Games, Inc. All Rights Reserved.
 
 #pragma once
-#include "GameFramework/HUD.h"
 #include "ShowFlags.h"
 #include "GameplayDebuggingHUDComponent.generated.h"
 
@@ -69,7 +68,7 @@ protected:
 private:
 	// local player related draw from PostRender
 	void DrawDebugComponentData(APlayerController* PC, class UGameplayDebuggingComponent *DebugComponent);
-	AGameplayDebuggingReplicator* GetDebuggingReplicator() const;
+	AGameplayDebuggingReplicator* GetDebuggingReplicator();
 
 protected:
 	UPROPERTY(Transient)
diff --git a/Engine/Source/Developer/GameplayDebugger/Classes/GameplayDebuggingReplicator.h b/Engine/Source/Developer/GameplayDebugger/Classes/GameplayDebuggingReplicator.h
--- a/Engine/Source/Developer/GameplayDebugger/Classes/GameplayDebuggingReplicator.h
+++ b/Engine/Source/Developer/GameplayDebugger/Classes/GameplayDebuggingReplicator.h
@@ -10,22 +10,6 @@
 class UGameplayDebuggingComponent;
 class AGameplayDebuggingHUDComponent;
 
-struct FDebugContext
-{
-	explicit FDebugContext(APlayerController* PC) : bEnabledTargetSelection(false), PlayerOwner(PC) {}
-	FDebugContext(const FDebugContext& DC) : bEnabledTargetSelection(DC.bEnabledTargetSelection), PlayerOwner(DC.PlayerOwner), DebugTarget(DC.DebugTarget) {}
-
-	uint32 bEnabledTargetSelection : 1;
-	TWeakObjectPtr<APlayerController> PlayerOwner;
-	TWeakObjectPtr<AActor>	DebugTarget;
-};
-
-FORCEINLINE bool operator==(const FDebugContext& Left, const FDebugContext& Right)
-{
-	return Left.PlayerOwner.Get() == Right.PlayerOwner.Get();
-}
-
-
 UCLASS(config=Engine, NotBlueprintable, Transient)
 class GAMEPLAYDEBUGGER_API AGameplayDebuggingReplicator : public AActor
 {
@@ -40,6 +24,9 @@ class GAMEPLAYDEBUGGER_API AGameplayDebuggingReplicator : public AActor
 	UPROPERTY(Replicated, Transient)
 	UGameplayDebuggingComponent* DebugComponent;
 
+	UPROPERTY(Replicated, Transient)
+	APlayerController* LocalPlayerOwner;
+
 	UFUNCTION(reliable, server, WithValidation)
 	void ServerReplicateMessage(class AActor* Actor, uint32 InMessage, uint32 DataView = 0);
 
@@ -49,12 +36,6 @@ class GAMEPLAYDEBUGGER_API AGameplayDebuggingReplicator : public AActor
 	UFUNCTION(reliable, server, WithValidation)
 	void ServerEnableTargetSelection(bool bEnable, APlayerController* Context);
 
-	UFUNCTION(reliable, server, WithValidation)
-	void ServerRegisterClient(APlayerController *PlayerController);
-
-	UFUNCTION(reliable, server, WithValidation)
-	void ServerUnregisterClient(APlayerController* PlayerController);
-
 	virtual class UNetConnection* GetNetConnection() override;
 
 	virtual bool IsNetRelevantFor(class APlayerController* RealViewer, AActor* Viewer, const FVector& SrcLocation) override;
@@ -63,17 +44,18 @@ class GAMEPLAYDEBUGGER_API AGameplayDebuggingReplicator : public AActor
 
 	virtual void BeginPlay() override;
 
-	FORCEINLINE UGameplayDebuggingComponent* GetDebugComponent()
-	{
-		return DebugComponent;
-	}
+	UGameplayDebuggingComponent* GetDebugComponent();
 
 	bool IsToolCreated();
-	void CreateTool(APlayerController *PlayerController);
+	void CreateTool();
 	void EnableTool();
 	bool IsDrawEnabled();
 	void EnableDraw(bool bEnable);
-	TArray<FDebugContext>& GetClients() { return Clients;  }
+
+	void SetLocalPlayerOwner(APlayerController* PC) { LocalPlayerOwner = PC; }
+	APlayerController* GetLocalPlayerOwner() { return LocalPlayerOwner; }
+
+	uint32 DebuggerShowFlags;
 
 protected:
 	void OnDebugAIDelegate(class UCanvas* Canvas, class APlayerController* PC);
@@ -88,7 +70,5 @@ private:
 	TWeakObjectPtr<UClass> DebugComponentHUDClass;
 
 	TWeakObjectPtr<AGameplayDebuggingHUDComponent>	DebugRenderer;
-	TArray<FDebugContext> Clients;
-	TWeakObjectPtr<APlayerController> LocalPlayerOwner;
 	TWeakObjectPtr<AActor>	LastSelectedActorInSimulation;
 };
diff --git a/Engine/Source/Developer/GameplayDebugger/Classes/GameplayDebuggingTypes.h b/Engine/Source/Developer/GameplayDebugger/Classes/GameplayDebuggingTypes.h
--- a/Engine/Source/Developer/GameplayDebugger/Classes/GameplayDebuggingTypes.h
+++ b/Engine/Source/Developer/GameplayDebugger/Classes/GameplayDebuggingTypes.h
@@ -28,22 +28,29 @@ namespace EAIDebugDrawDataView
 
 struct GAMEPLAYDEBUGGER_API FGameplayDebuggerSettings
 {
-	static void SetFlag(uint8 Flag)
+	explicit FGameplayDebuggerSettings(uint32& Flags)
+		: DebuggerShowFlags(Flags)
+	{
+
+	}
+	void SetFlag(uint8 Flag)
 	{
 		DebuggerShowFlags |= (1 << Flag);
 	}
 
-	static void ClearFlag(uint8 Flag)
+	void ClearFlag(uint8 Flag)
 	{
 		DebuggerShowFlags &= ~(1 << Flag);
 	}
 
-	static bool CheckFlag(uint8 Flag)
+	bool CheckFlag(uint8 Flag)
 	{
 		return (DebuggerShowFlags & (1 << Flag)) != 0;
 	}
 
-	static uint32 DebuggerShowFlags;
+	uint32 &DebuggerShowFlags;
 	static uint32 ShowFlagIndex;
 };
 
+GAMEPLAYDEBUGGER_API
+FGameplayDebuggerSettings  GameplayDebuggerSettings(class AGameplayDebuggingReplicator* Replicator = NULL);
diff --git a/Engine/Source/Developer/GameplayDebugger/Private/GameplayDebugger.cpp b/Engine/Source/Developer/GameplayDebugger/Private/GameplayDebugger.cpp
--- a/Engine/Source/Developer/GameplayDebugger/Private/GameplayDebugger.cpp
+++ b/Engine/Source/Developer/GameplayDebugger/Private/GameplayDebugger.cpp
@@ -5,9 +5,14 @@
 #include "GameplayDebuggingHUDComponent.h"
 #include "GameplayDebuggingReplicator.h"
 
-uint32 FGameplayDebuggerSettings::DebuggerShowFlags = (1 << EAIDebugDrawDataView::Basic) | (1 << EAIDebugDrawDataView::OverHead);
 uint32 FGameplayDebuggerSettings::ShowFlagIndex = 0;
 
+FGameplayDebuggerSettings GameplayDebuggerSettings(class AGameplayDebuggingReplicator* Replicator)
+{
+	static uint32 Settings = (1 << EAIDebugDrawDataView::Basic) | (1 << EAIDebugDrawDataView::OverHead);
+	return FGameplayDebuggerSettings(Replicator == NULL ? Settings : Replicator->DebuggerShowFlags);
+}
+
 class FGameplayDebugger : public IGameplayDebugger
 {
 public:
@@ -55,15 +60,7 @@ void FGameplayDebugger::ShutdownModule()
 
 void FGameplayDebugger::WorldAdded(UWorld* InWorld)
 {
-#if !(UE_BUILD_SHIPPING || UE_BUILD_TEST)
-	if (InWorld->GetNetMode() < ENetMode::NM_Client)
-	{
-		FActorSpawnParameters SpawnInfo;
-		SpawnInfo.bNoCollisionFail = true;
-		SpawnInfo.Name = TEXT("GameplayDebuggingReplicator");
-		AGameplayDebuggingReplicator *DestActor = InWorld->SpawnActor<AGameplayDebuggingReplicator>(SpawnInfo);
-	}
-#endif
+
 }
 
 void FGameplayDebugger::WorldDestroyed(UWorld* InWorld)
@@ -73,7 +70,19 @@ void FGameplayDebugger::WorldDestroyed(UWorld* InWorld)
 
 void FGameplayDebugger::OnLevelActorAdded(AActor* InActor)
 {
-
+#if !(UE_BUILD_SHIPPING || UE_BUILD_TEST)
+	APlayerController* PC = Cast<APlayerController>(InActor);
+	if (PC && PC->GetNetMode() < ENetMode::NM_Client)
+	{
+		UWorld* World = PC->GetWorld();
+		FActorSpawnParameters SpawnInfo;
+		SpawnInfo.bNoCollisionFail = true;
+		SpawnInfo.Name = *FString::Printf(TEXT("GameplayDebuggingReplicator_%s"), *PC->GetName());
+		AGameplayDebuggingReplicator *DestActor = World->SpawnActor<AGameplayDebuggingReplicator>(SpawnInfo);
+		DestActor->SetLocalPlayerOwner(PC);
+		DestActor->SetReplicates(true);
+	}
+#endif
 }
 
 void FGameplayDebugger::OnLevelActorDeleted(AActor* InActor)
diff --git a/Engine/Source/Developer/GameplayDebugger/Private/GameplayDebuggerExec.cpp b/Engine/Source/Developer/GameplayDebugger/Private/GameplayDebuggerExec.cpp
--- a/Engine/Source/Developer/GameplayDebugger/Private/GameplayDebuggerExec.cpp
+++ b/Engine/Source/Developer/GameplayDebugger/Private/GameplayDebuggerExec.cpp
@@ -1,10 +1,16 @@
 // Copyright 1998-2014 Epic Games, Inc. All Rights Reserved.
 
 #include "GameplayDebuggerPrivate.h"
-#include "GameplayDebuggingControllerComponent.h"
 #if !(UE_BUILD_SHIPPING || UE_BUILD_TEST)
+#include "GameplayDebuggingControllerComponent.h"
 #include "Misc/CoreMisc.h"
 
+#if WITH_EDITOR
+#	include "Editor.h"
+#	include "LevelEditorViewport.h"
+#endif //WITH_EDITOR
+
+
 struct FGameplayDebuggerExec : public FSelfRegisteringExec
 {
 	FGameplayDebuggerExec()
@@ -55,10 +61,28 @@ bool FGameplayDebuggerExec::Exec(UWorld* Inworld, const TCHAR* Cmd, FOutputDevic
 
 	if (FParse::Command(&Cmd, TEXT("cheat EnableGDT")))
 	{
-		if (Inworld->GetNetMode() != NM_DedicatedServer && GetDebuggingReplicator(Inworld).IsValid() && !CachedDebuggingReplicator->IsToolCreated())
+		if (Inworld->GetNetMode() != NM_DedicatedServer)
 		{
-			CachedDebuggingReplicator->CreateTool(PC);
-			CachedDebuggingReplicator->EnableTool();
+			AGameplayDebuggingReplicator* Replicator = NULL;
+			for (FActorIterator It(Inworld); It; ++It)
+			{
+				AActor* A = *It;
+				if (A && A->IsA(AGameplayDebuggingReplicator::StaticClass()) && !A->IsPendingKill())
+				{
+					Replicator = Cast<AGameplayDebuggingReplicator>(A);
+					if (Replicator && Replicator->GetLocalPlayerOwner() == PC)
+					{
+						break;
+					}
+				}
+			}
+
+			if (Replicator && !Replicator->IsToolCreated())
+			{
+				Replicator->CreateTool();
+				Replicator->EnableTool();
+				bHandled = true;
+			}
 		}
 	}
 	else if (FParse::Command(&Cmd, TEXT("ToggleGameplayDebugView")))
@@ -74,21 +98,45 @@ bool FGameplayDebuggerExec::Exec(UWorld* Inworld, const TCHAR* Cmd, FOutputDevic
 			}
 		}
 
+		const bool bActivePIE = Inworld && Inworld->IsGameWorld();
+		AGameplayDebuggingReplicator* DebuggingReplicator = NULL;
+		if (bActivePIE)
+		{
+			for (FActorIterator It(Inworld); It; ++It)
+			{
+				AActor* A = *It;
+				if (A && A->IsA(AGameplayDebuggingReplicator::StaticClass()) && !A->IsPendingKill())
+				{
+					AGameplayDebuggingReplicator* Replicator = Cast<AGameplayDebuggingReplicator>(A);
+					if (Replicator->GetLocalPlayerOwner() == PC)
+					{
+						DebuggingReplicator = Replicator;
+						break;
+					}
+				}
+			}
+		}
+
+		FGameplayDebuggerSettings DebuggerSettings = GameplayDebuggerSettings(DebuggingReplicator);
 		FString InViewName = FParse::Token(Cmd, 0);
 		int32 ViewIndex = ViewNames.Find(InViewName);
 		FGameplayDebuggerSettings::ShowFlagIndex = FEngineShowFlags::FindIndexByName(TEXT("GameplayDebug"));
 		if (ViewIndex != INDEX_NONE)
 		{
-			FGameplayDebuggerSettings::CheckFlag(ViewIndex) ? FGameplayDebuggerSettings::ClearFlag(ViewIndex) : FGameplayDebuggerSettings::SetFlag(ViewIndex);
+			DebuggerSettings.CheckFlag(ViewIndex) ? DebuggerSettings.ClearFlag(ViewIndex) : DebuggerSettings.SetFlag(ViewIndex);
+			if (bActivePIE)
+			{
+				GameplayDebuggerSettings().CheckFlag(ViewIndex) ? GameplayDebuggerSettings().ClearFlag(ViewIndex) : GameplayDebuggerSettings().SetFlag(ViewIndex);
+			}
 #if WITH_EDITOR
 			if (ViewIndex == EAIDebugDrawDataView::EQS && GCurrentLevelEditingViewportClient)
 			{
-				GCurrentLevelEditingViewportClient->EngineShowFlags.SetSingleFlag(FGameplayDebuggerSettings::ShowFlagIndex, FGameplayDebuggerSettings::CheckFlag(ViewIndex));
+				GCurrentLevelEditingViewportClient->EngineShowFlags.SetSingleFlag(FGameplayDebuggerSettings::ShowFlagIndex, DebuggerSettings.CheckFlag(ViewIndex));
 			}
 #endif
 			PC->ClientMessage(FString::Printf(TEXT("View %s %s")
 				, *InViewName
-				, FGameplayDebuggerSettings::CheckFlag(ViewIndex) ? TEXT("enabled") : TEXT("disabled")));
+				, DebuggerSettings.CheckFlag(ViewIndex) ? TEXT("enabled") : TEXT("disabled")));
 		}
 		else
 		{
diff --git a/Engine/Source/Developer/GameplayDebugger/Private/GameplayDebuggerPrivate.h b/Engine/Source/Developer/GameplayDebugger/Private/GameplayDebuggerPrivate.h
--- a/Engine/Source/Developer/GameplayDebugger/Private/GameplayDebuggerPrivate.h
+++ b/Engine/Source/Developer/GameplayDebugger/Private/GameplayDebuggerPrivate.h
@@ -36,11 +36,4 @@
 #include "GameplayDebuggingControllerComponent.h"
 #include "GameplayDebuggingComponent.h"
 
-#if WITH_EDITOR
-#include "Editor.h"
-#include "EditorViewportClient.h"
-#include "LevelEditorViewport.h"
-#include "UnrealEd.h"
-#endif //WITH_EDITOR
-
 DECLARE_LOG_CATEGORY_EXTERN(LogGDT, Warning, All);
diff --git a/Engine/Source/Developer/GameplayDebugger/Private/GameplayDebuggingComponent.cpp b/Engine/Source/Developer/GameplayDebugger/Private/GameplayDebuggingComponent.cpp
--- a/Engine/Source/Developer/GameplayDebugger/Private/GameplayDebuggingComponent.cpp
+++ b/Engine/Source/Developer/GameplayDebugger/Private/GameplayDebuggingComponent.cpp
@@ -16,10 +16,6 @@
 
 #include "GameplayDebuggingComponent.h"
 
-#if WITH_EDITOR
-//#include "UnrealEd.h"
-#endif
-
 DEFINE_LOG_CATEGORY(LogGDT);
 
 //----------------------------------------------------------------------//
@@ -250,32 +246,23 @@ void UGameplayDebuggingComponent::TickComponent(float DeltaTime, enum ELevelTick
 			AGameplayDebuggingReplicator* Replicator = Cast<AGameplayDebuggingReplicator>(GetOwner());
 			if (Replicator)
 			{
-				TArray<FDebugContext>& Clients = Replicator->GetClients();
-				for (int32 Index = 0; Index < Clients.Num(); ++Index)
+				if (Replicator->GetLocalPlayerOwner())
 				{
-					if (Clients[Index].bEnabledTargetSelection && Clients[Index].PlayerOwner.IsValid())
-					{
-						SelectTargetToDebug(Clients[Index]);
-
-						if (Clients[Index].DebugTarget.IsValid())
-						{
-							SetActorToDebug(Clients[Index].DebugTarget.Get());
-							CollectDataToReplicate(true);
-						}
-					}
-					break;
+					SelectTargetToDebug();
 				}
 			}
 		}
+		CollectDataToReplicate(true);
 	}
 
 #endif //!(UE_BUILD_SHIPPING || UE_BUILD_TEST)
 }
 
-void UGameplayDebuggingComponent::SelectTargetToDebug(FDebugContext& Context)
+void UGameplayDebuggingComponent::SelectTargetToDebug()
 {
 #if !(UE_BUILD_SHIPPING || UE_BUILD_TEST)
-	APlayerController* MyPC = Context.PlayerOwner.Get();
+	AGameplayDebuggingReplicator* Replicator = Cast<AGameplayDebuggingReplicator>(GetOwner());
+	APlayerController* MyPC = Replicator->GetLocalPlayerOwner();
 
 	if (MyPC )
 	{
@@ -329,11 +316,10 @@ void UGameplayDebuggingComponent::SelectTargetToDebug(FDebugContext& Context)
 		}
 
 		BestTarget = BestTarget == NULL ? PossibleTarget : BestTarget;
-		if (BestTarget != NULL)
+		if (BestTarget != NULL && BestTarget != GetSelectedActor())
 		{
 			//always update component for best target
-			Context.DebugTarget = Cast<AActor>(BestTarget);
-			SetActorToDebug(Context.DebugTarget.Get());
+			SetActorToDebug(Cast<AActor>(BestTarget));
 			SelectForDebugging(true);
 			ServerReplicateData(EDebugComponentMessage::ActivateReplication, EAIDebugDrawDataView::Empty);
 		}
diff --git a/Engine/Source/Developer/GameplayDebugger/Private/GameplayDebuggingControllerComponent.cpp b/Engine/Source/Developer/GameplayDebugger/Private/GameplayDebuggingControllerComponent.cpp
--- a/Engine/Source/Developer/GameplayDebugger/Private/GameplayDebuggingControllerComponent.cpp
+++ b/Engine/Source/Developer/GameplayDebugger/Private/GameplayDebuggingControllerComponent.cpp
@@ -35,7 +35,7 @@ void UGameplayDebuggingControllerComponent::OnRegister()
 	BindActivationKeys();
 }
 
-AGameplayDebuggingReplicator* UGameplayDebuggingControllerComponent::GetDebuggingReplicator()
+AGameplayDebuggingReplicator* UGameplayDebuggingControllerComponent::GetDebuggingReplicator() const
 {
 	return Cast<AGameplayDebuggingReplicator>(GetOuter());
 }
@@ -83,17 +83,17 @@ void UGameplayDebuggingControllerComponent::BeginDestroy()
 
 bool UGameplayDebuggingControllerComponent::IsViewActive(EAIDebugDrawDataView::Type View) const
 {
-	return FGameplayDebuggerSettings::CheckFlag(View);
+	return GameplayDebuggerSettings(GetDebuggingReplicator()).CheckFlag(View);
 }
 
 uint32 UGameplayDebuggingControllerComponent::GetActiveViews()
 {
-	return FGameplayDebuggerSettings::DebuggerShowFlags;
+	return GameplayDebuggerSettings(GetDebuggingReplicator()).DebuggerShowFlags;
 }
 
 void UGameplayDebuggingControllerComponent::SetActiveViews(uint32 InActiveViews)
 {
-	FGameplayDebuggerSettings::DebuggerShowFlags = InActiveViews;
+	GameplayDebuggerSettings(GetDebuggingReplicator()).DebuggerShowFlags = InActiveViews;
 
 	if (DebugAITargetActor && GetDebuggingReplicator())
 	{
@@ -107,7 +107,7 @@ void UGameplayDebuggingControllerComponent::SetActiveViews(uint32 InActiveViews)
 
 void UGameplayDebuggingControllerComponent::EnableActiveView(EAIDebugDrawDataView::Type View, bool bEnable)
 {
-	bEnable ? FGameplayDebuggerSettings::SetFlag(View) : FGameplayDebuggerSettings::ClearFlag(View);
+	bEnable ? GameplayDebuggerSettings(GetDebuggingReplicator()).SetFlag(View) : GameplayDebuggerSettings(GetDebuggingReplicator()).ClearFlag(View);
 
 	if (DebugAITargetActor && GetDebuggingReplicator())
 	{
@@ -254,7 +254,7 @@ void UGameplayDebuggingControllerComponent::ToggleAIDebugView_SetView0()
 	{
 		if (UGameplayDebuggingComponent* OwnerComp = GetDebuggingReplicator() ? GetDebuggingReplicator()->GetDebugComponent() : NULL)
 		{
-			FGameplayDebuggerSettings::DebuggerShowFlags ^= 1 << EAIDebugDrawDataView::NavMesh;
+			GameplayDebuggerSettings(GetDebuggingReplicator()).DebuggerShowFlags ^= 1 << EAIDebugDrawDataView::NavMesh;
 
 			if (IsViewActive(EAIDebugDrawDataView::NavMesh))
 			{
diff --git a/Engine/Source/Developer/GameplayDebugger/Private/GameplayDebuggingHUDComponent.cpp b/Engine/Source/Developer/GameplayDebugger/Private/GameplayDebuggingHUDComponent.cpp
--- a/Engine/Source/Developer/GameplayDebugger/Private/GameplayDebuggingHUDComponent.cpp
+++ b/Engine/Source/Developer/GameplayDebugger/Private/GameplayDebuggingHUDComponent.cpp
@@ -39,9 +39,9 @@ void AGameplayDebuggingHUDComponent::Render()
 	PrintAllData();
 }
 
-AGameplayDebuggingReplicator* AGameplayDebuggingHUDComponent::GetDebuggingReplicator() const
+AGameplayDebuggingReplicator* AGameplayDebuggingHUDComponent::GetDebuggingReplicator()
 {
-	if (CachedDebuggingReplicator.IsValid())
+	if (CachedDebuggingReplicator.IsValid() && CachedDebuggingReplicator->GetLocalPlayerOwner() == PlayerOwner)
 	{
 		return CachedDebuggingReplicator.Get();
 	}
@@ -51,7 +51,12 @@ AGameplayDebuggingReplicator* AGameplayDebuggingHUDComponent::GetDebuggingReplic
 		AActor* A = *It;
 		if (A && A->IsA(AGameplayDebuggingReplicator::StaticClass()) && !A->IsPendingKill())
 		{
-			return Cast<AGameplayDebuggingReplicator>(A);
+			AGameplayDebuggingReplicator* Replicator = Cast<AGameplayDebuggingReplicator>(A);
+			if (Replicator->GetLocalPlayerOwner() == PlayerOwner)
+			{
+				CachedDebuggingReplicator = Replicator;
+				return Replicator;
+			}
 		}
 	}
 
@@ -146,10 +151,10 @@ void AGameplayDebuggingHUDComponent::DrawMenu(const float X, const float Y, clas
 			float XPos = 20.0f;
 			for (int32 i = 0; i < Categories.Num(); i++)
 			{
-			const bool bIsActive = FGameplayDebuggerSettings::CheckFlag(Categories[i].View) ? true : false;
-			const bool bIsDisabled = Categories[i].View == EAIDebugDrawDataView::NavMesh ? false : (DebugComponent && DebugComponent->GetSelectedActor() ? false: true);
+				const bool bIsActive = GameplayDebuggerSettings(GetDebuggingReplicator()).CheckFlag(Categories[i].View) ? true : false;
+				const bool bIsDisabled = Categories[i].View == EAIDebugDrawDataView::NavMesh ? false : (DebugComponent && DebugComponent->GetSelectedActor() ? false: true);
 
-			PrintString(DefaultContext, bIsDisabled ? (bIsActive ? FColorList::DarkGreen  : FColorList::LightGrey) : (bIsActive ? FColorList::Green : FColorList::White), Categories[i].Desc, XPos, 20);
+				PrintString(DefaultContext, bIsDisabled ? (bIsActive ? FColorList::DarkGreen  : FColorList::LightGrey) : (bIsActive ? FColorList::Green : FColorList::White), Categories[i].Desc, XPos, 20);
 				XPos += CategoriesWidth[i];
 			}
 		DefaultContext.Font = OldFont;
@@ -172,34 +177,35 @@ void AGameplayDebuggingHUDComponent::DrawDebugComponentData(APlayerController* M
 
 	OverHeadContext = FPrintContext(GEngine->GetSmallFont(), Canvas, ScreenLoc.X, ScreenLoc.Y);
 
-	if (FGameplayDebuggerSettings::CheckFlag(EAIDebugDrawDataView::OverHead) || EngineShowFlags.DebugAI)
+	FGameplayDebuggerSettings DebuggerSettings = GameplayDebuggerSettings(GetDebuggingReplicator());
+	if (DebuggerSettings.CheckFlag(EAIDebugDrawDataView::OverHead) || EngineShowFlags.DebugAI)
 	{
 		DrawOverHeadInformation(MyPC, DebugComponent);
 	}
 
-	if (FGameplayDebuggerSettings::CheckFlag(EAIDebugDrawDataView::NavMesh))
+	if (DebuggerSettings.CheckFlag(EAIDebugDrawDataView::NavMesh))
 	{
 		DrawNavMeshSnapshot(MyPC, DebugComponent);
 	}
 
 	if (DebugComponent->GetSelectedActor() && DebugComponent->IsSelected())
 	{
-		if (FGameplayDebuggerSettings::CheckFlag(EAIDebugDrawDataView::Basic) || EngineShowFlags.DebugAI)
+		if (DebuggerSettings.CheckFlag(EAIDebugDrawDataView::Basic) || EngineShowFlags.DebugAI)
 		{
 			DrawBasicData(MyPC, DebugComponent);
 		}
 
-		if (FGameplayDebuggerSettings::CheckFlag(EAIDebugDrawDataView::BehaviorTree))
+		if (DebuggerSettings.CheckFlag(EAIDebugDrawDataView::BehaviorTree))
 		{
 			DrawBehaviorTreeData(MyPC, DebugComponent);
 		}
 
-		if (FGameplayDebuggerSettings::CheckFlag(EAIDebugDrawDataView::EQS))
+		if (DebuggerSettings.CheckFlag(EAIDebugDrawDataView::EQS))
 		{
 			DrawEQSData(MyPC, DebugComponent);
 		}
 
-		if (FGameplayDebuggerSettings::CheckFlag(EAIDebugDrawDataView::Perception) || EngineShowFlags.DebugAI)
+		if (DebuggerSettings.CheckFlag(EAIDebugDrawDataView::Perception) || EngineShowFlags.DebugAI)
 		{
 			DrawPerception(MyPC, DebugComponent);
 		}
diff --git a/Engine/Source/Developer/GameplayDebugger/Private/GameplayDebuggingReplicator.cpp b/Engine/Source/Developer/GameplayDebugger/Private/GameplayDebuggingReplicator.cpp
--- a/Engine/Source/Developer/GameplayDebugger/Private/GameplayDebuggingReplicator.cpp
+++ b/Engine/Source/Developer/GameplayDebugger/Private/GameplayDebuggingReplicator.cpp
@@ -1,6 +1,7 @@
 // Copyright 1998-2014 Epic Games, Inc. All Rights Reserved.
 
 #include "GameplayDebuggerPrivate.h"
+#include "GameFramework/HUD.h"
 #include "GameplayDebuggingComponent.h"
 #include "GameplayDebuggingHUDComponent.h"
 #include "GameplayDebuggingReplicator.h"
@@ -16,6 +17,7 @@ AGameplayDebuggingReplicator::AGameplayDebuggingReplicator(const class FPostCons
 {
 	TSubobjectPtr<USceneComponent> SceneComponent = PCIP.CreateDefaultSubobject<USceneComponent>(this, TEXT("SceneComponent"));
 	RootComponent = SceneComponent;
+
 #if WITH_EDITORONLY_DATA
 	SetTickableWhenPaused(true);
 	SetIsTemporarilyHiddenInEditor(true);
@@ -25,6 +27,17 @@ AGameplayDebuggingReplicator::AGameplayDebuggingReplicator(const class FPostCons
 	bHiddenEd = true;
 	bEditable = false;
 #endif
+
+	DebuggerShowFlags =  GameplayDebuggerSettings().DebuggerShowFlags;
+
+	if (!HasAnyFlags(RF_ClassDefaultObject))
+	{
+		SetActorTickEnabled(true);
+
+		bReplicates = false;
+		SetRemoteRoleForBackwardsCompat(ROLE_SimulatedProxy);
+		SetReplicates(true);
+	}
 }
 
 void AGameplayDebuggingReplicator::GetLifetimeReplicatedProps(TArray< FLifetimeProperty > & OutLifetimeProps) const
@@ -32,25 +45,21 @@ void AGameplayDebuggingReplicator::GetLifetimeReplicatedProps(TArray< FLifetimeP
 	Super::GetLifetimeReplicatedProps(OutLifetimeProps);
 
 #if !(UE_BUILD_SHIPPING || UE_BUILD_TEST)
-	DOREPLIFETIME(AGameplayDebuggingReplicator, DebugComponent);
+		DOREPLIFETIME_CONDITION(AGameplayDebuggingReplicator, DebugComponent, COND_OwnerOnly);
+		DOREPLIFETIME_CONDITION(AGameplayDebuggingReplicator, LocalPlayerOwner, COND_OwnerOnly);
 #endif
 }
 
 bool AGameplayDebuggingReplicator::IsNetRelevantFor(class APlayerController* RealViewer, AActor* Viewer, const FVector& SrcLocation)
 {
-	if (Clients.Num() == 0)
-	{
-		return true;// relevant to everyone for now
-	}
-
-	return Clients.Contains(FDebugContext(RealViewer));
+	return LocalPlayerOwner == RealViewer;
 }
 
 void AGameplayDebuggingReplicator::PostInitializeComponents()
 {
 	Super::PostInitializeComponents();
 
-	if (!HasAnyFlags(RF_ClassDefaultObject) && GetWorld() && GetNetMode() != ENetMode::NM_DedicatedServer)
+	if (GetWorld() && GetNetMode() != ENetMode::NM_DedicatedServer)
 	{
 		if (GIsEditor)
 		{
@@ -66,44 +75,41 @@ void AGameplayDebuggingReplicator::PostInitializeComponents()
 				DebugComponentHUDClass = AGameplayDebuggingHUDComponent::StaticClass();
 			}
 		}
+	}
 
+	if (!DebugComponentClass.IsValid() && GetWorld() && GetNetMode() < ENetMode::NM_Client)
+	{
+		DebugComponentClass = StaticLoadClass(UGameplayDebuggingComponent::StaticClass(), NULL, *DebugComponentClassName, NULL, LOAD_None, NULL);
 		if (!DebugComponentClass.IsValid())
 		{
-			DebugComponentClass = StaticLoadClass(UGameplayDebuggingComponent::StaticClass(), NULL, *DebugComponentClassName, NULL, LOAD_None, NULL);
-			if (!DebugComponentClass.IsValid())
-			{
-				DebugComponentClass = UGameplayDebuggingComponent::StaticClass();
-			}
+			DebugComponentClass = UGameplayDebuggingComponent::StaticClass();
 		}
 	}
+}
+
+UGameplayDebuggingComponent* AGameplayDebuggingReplicator::GetDebugComponent()
+{
+	if (!DebugComponent && DebugComponentClass.IsValid() && GetNetMode() < ENetMode::NM_Client)
+	{
+		DebugComponent = ConstructObject<UGameplayDebuggingComponent>(DebugComponentClass.Get(), this);
+		DebugComponent->SetIsReplicated(true);
+		DebugComponent->RegisterComponent();
+		DebugComponent->Activate();
+	}
 
-	SetActorTickEnabled(true);
+	return DebugComponent;
 }
 
 void AGameplayDebuggingReplicator::BeginPlay()
 {
 	Super::BeginPlay();
-	if (Role == ROLE_Authority)
-	{
-		bReplicates = false;
-		SetRemoteRoleForBackwardsCompat(ROLE_SimulatedProxy);
-		SetReplicates(true);
-
-		if (DebugComponent == NULL)
-		{
-			DebugComponent = ConstructObject<UGameplayDebuggingComponent>(DebugComponentClass.Get(), this);
-			DebugComponent->SetIsReplicated(true);
-			DebugComponent->RegisterComponent();
-			DebugComponent->Activate();
-		}
-	}
 }
 
 class UNetConnection* AGameplayDebuggingReplicator::GetNetConnection()
 {
-	if (GetWorld() && GetWorld()->GetFirstPlayerController())
+	if (LocalPlayerOwner)
 	{
-		return GetWorld()->GetFirstPlayerController()->GetNetConnection();
+		return LocalPlayerOwner->GetNetConnection();
 	}
 
 	return NULL;
@@ -116,15 +122,9 @@ bool AGameplayDebuggingReplicator::ServerEnableTargetSelection_Validate(bool, AP
 
 void AGameplayDebuggingReplicator::ServerEnableTargetSelection_Implementation(bool bEnable, APlayerController* Context)
 {
-	if (DebugComponent)
+	if (GetDebugComponent())
 	{
-		uint32 Index = Clients.Find(FDebugContext(Context));
-		if (Index != INDEX_NONE)
-		{
-			Clients[Index].bEnabledTargetSelection = bEnable;
-		}
-
-		DebugComponent->ServerEnableTargetSelection(true);
+		GetDebugComponent()->ServerEnableTargetSelection(true);
 	}
 }
 
@@ -146,39 +146,13 @@ bool AGameplayDebuggingReplicator::ServerReplicateMessage_Validate(class AActor*
 void AGameplayDebuggingReplicator::ServerReplicateMessage_Implementation(class  AActor* Actor, uint32 InMessage, uint32 DataView)
 {
 #if !(UE_BUILD_SHIPPING || UE_BUILD_TEST)
-	if (DebugComponent)
+	if (GetDebugComponent())
 	{
-		DebugComponent->ServerReplicateData((EDebugComponentMessage::Type)InMessage, (EAIDebugDrawDataView::Type)DataView);
+		GetDebugComponent()->ServerReplicateData((EDebugComponentMessage::Type)InMessage, (EAIDebugDrawDataView::Type)DataView);
 	}
 #endif //!(UE_BUILD_SHIPPING || UE_BUILD_TEST)
 }
 
-bool AGameplayDebuggingReplicator::ServerRegisterClient_Validate(APlayerController*)
-{
-	return true;
-}
-
-void AGameplayDebuggingReplicator::ServerRegisterClient_Implementation(APlayerController* NewClient)
-{
-	if (!Clients.Contains(FDebugContext(NewClient)))
-	{
-		Clients.AddUnique(FDebugContext(NewClient));
-	}
-}
-
-bool AGameplayDebuggingReplicator::ServerUnregisterClient_Validate(APlayerController*)
-{
-	return true;
-}
-
-void AGameplayDebuggingReplicator::ServerUnregisterClient_Implementation(APlayerController* OldClient)
-{
-	if (Clients.Contains(FDebugContext(OldClient)))
-	{
-		Clients.Remove(FDebugContext(OldClient));
-	}
-}
-
 bool AGameplayDebuggingReplicator::IsDrawEnabled()
 {
 	return bEnabledDraw && GetWorld() && GetNetMode() != ENetMode::NM_DedicatedServer;
@@ -188,7 +162,7 @@ void AGameplayDebuggingReplicator::EnableDraw(bool bEnable)
 {
 	bEnabledDraw = bEnable;
 
-	if (AHUD* const GameHUD = LocalPlayerOwner.IsValid() ? LocalPlayerOwner->GetHUD() : NULL)
+	if (AHUD* const GameHUD = LocalPlayerOwner ? LocalPlayerOwner->GetHUD() : NULL)
 	{
 		GameHUD->bShowHUD = bEnable ? false : true;
 	}
@@ -198,20 +172,18 @@ void AGameplayDebuggingReplicator::EnableDraw(bool bEnable)
 bool AGameplayDebuggingReplicator::IsToolCreated()
 {
 	UGameplayDebuggingControllerComponent*  GDC = FindComponentByClass<UGameplayDebuggingControllerComponent>();
-	return LocalPlayerOwner.IsValid() && GDC;
+	return LocalPlayerOwner && GDC;
 }
 
-void AGameplayDebuggingReplicator::CreateTool(APlayerController *PlayerController)
+void AGameplayDebuggingReplicator::CreateTool()
 {
 	if (GetWorld() && GetNetMode() != ENetMode::NM_DedicatedServer)
 	{
-		LocalPlayerOwner = PlayerController;
-		ServerRegisterClient( PlayerController );
 		UGameplayDebuggingControllerComponent*  GDC = FindComponentByClass<UGameplayDebuggingControllerComponent>();
 		if (!GDC)
 		{
 			GDC = ConstructObject<UGameplayDebuggingControllerComponent>(UGameplayDebuggingControllerComponent::StaticClass(), this);
-			GDC->SetPlayerOwner(PlayerController);
+			GDC->SetPlayerOwner(LocalPlayerOwner);
 			GDC->RegisterComponent();
 		}
 	}
@@ -252,7 +224,7 @@ void AGameplayDebuggingReplicator::OnDebugAIDelegate(class UCanvas* Canvas, clas
 	}
 
 	UWorld* World = GetWorld();
-	if (World && DebugComponent && DebugComponent->GetOwnerRole() == ROLE_Authority)
+	if (World && GetDebugComponent() && GetDebugComponent()->GetOwnerRole() == ROLE_Authority)
 	{
 		bool bBroadcastedNewSelection = false;
 
@@ -262,11 +234,9 @@ void AGameplayDebuggingReplicator::OnDebugAIDelegate(class UCanvas* Canvas, clas
 			AActor* NewTarget = Cast<AActor>(*Iterator);
 
 			//We needs to collect data manually in Simulate
-			FDebugContext Context(LocalPlayerOwner.Get());
-			Context.DebugTarget = NewTarget;
-			DebugComponent->SetActorToDebug(NewTarget);
-			DebugComponent->SelectForDebugging(NewTarget->IsSelected());
-			DebugComponent->CollectDataToReplicate(NewTarget->IsSelected());
+			GetDebugComponent()->SetActorToDebug(NewTarget);
+			GetDebugComponent()->SelectForDebugging(NewTarget->IsSelected());
+			GetDebugComponent()->CollectDataToReplicate(NewTarget->IsSelected());
 
 			if (NewTarget->IsSelected() && LastSelectedActorInSimulation.Get() != NewTarget)
 			{
@@ -310,7 +280,8 @@ void AGameplayDebuggingReplicator::DrawDebugDataDelegate(class UCanvas* Canvas,
 	}
 	LastDrawAtFrame = GFrameNumber;
 
-	if (!IsDrawEnabled())
+	const UGameplayDebuggingControllerComponent*  GDC = FindComponentByClass<UGameplayDebuggingControllerComponent>();
+	if (!IsDrawEnabled() || !GDC)
 	{
 		return;
 	}
@@ -320,10 +291,8 @@ void AGameplayDebuggingReplicator::DrawDebugDataDelegate(class UCanvas* Canvas,
 		for (FConstPawnIterator Iterator = GetWorld()->GetPawnIterator(); Iterator; ++Iterator)
 		{
 			AActor* NewTarget = Cast<AActor>(*Iterator);
-			FDebugContext Context(LocalPlayerOwner.Get());
-			Context.DebugTarget = NewTarget;
-			DebugComponent->SetActorToDebug(NewTarget);
-			DebugComponent->CollectDataToReplicate(true);
+			GetDebugComponent()->SetActorToDebug(NewTarget);
+			GetDebugComponent()->CollectDataToReplicate(true);
 		}
 	}
 
@@ -332,6 +301,12 @@ void AGameplayDebuggingReplicator::DrawDebugDataDelegate(class UCanvas* Canvas,
 
 void AGameplayDebuggingReplicator::DrawDebugData(class UCanvas* Canvas, class APlayerController* PC)
 {
+	const bool bAllowToDraw = Canvas && Canvas->SceneView && (Canvas->SceneView->ViewActor == LocalPlayerOwner->AcknowledgedPawn || Canvas->SceneView->ViewActor == LocalPlayerOwner->GetPawnOrSpectator());
+	if (!bAllowToDraw)
+	{
+		return;
+	}
+
 	if (!DebugRenderer.IsValid() && DebugComponentHUDClass.IsValid())
 	{
 		FActorSpawnParameters SpawnInfo;
@@ -341,6 +316,7 @@ void AGameplayDebuggingReplicator::DrawDebugData(class UCanvas* Canvas, class AP
 
 		DebugRenderer = GetWorld()->SpawnActor<AGameplayDebuggingHUDComponent>(DebugComponentHUDClass.Get(), SpawnInfo);
 		DebugRenderer->SetCanvas(Canvas);
+		DebugRenderer->SetPlayerOwner(LocalPlayerOwner);
 		DebugRenderer->SetWorld(GetWorld());
 	}
 
diff --git a/Engine/Source/Runtime/Engine/Private/Debug/DebugDrawService.cpp b/Engine/Source/Runtime/Engine/Private/Debug/DebugDrawService.cpp
--- a/Engine/Source/Runtime/Engine/Private/Debug/DebugDrawService.cpp
+++ b/Engine/Source/Runtime/Engine/Private/Debug/DebugDrawService.cpp
@@ -56,10 +56,11 @@ void UDebugDrawService::Draw(const FEngineShowFlags Flags, FViewport* Viewport,
 		CanvasObject = ConstructObject<UCanvas>(UCanvas::StaticClass(),GetTransientPackage(),TEXT("DebugCanvasObject"));
 		CanvasObject->AddToRoot();
 	}
-	
+
 	CanvasObject->Init(View->UnscaledViewRect.Width(), View->UnscaledViewRect.Height(), View);
 	CanvasObject->Update();	
 	CanvasObject->Canvas = Canvas;
+	CanvasObject->SetView(View);
 
 	// PreRender the player's view.
 	Draw(Flags, CanvasObject);	
