Fixed warning about using wrong precompiled header #UE4

diff --git a/Engine/Source/Developer/GameplayDebugger/Private/GameplayDebuggerExec.cpp b/Engine/Source/Developer/GameplayDebugger/Private/GameplayDebuggerExec.cpp
--- a/Engine/Source/Developer/GameplayDebugger/Private/GameplayDebuggerExec.cpp
+++ b/Engine/Source/Developer/GameplayDebugger/Private/GameplayDebuggerExec.cpp
@@ -6,8 +6,7 @@
 #include "Misc/CoreMisc.h"
 
 #if WITH_EDITOR
-#	include "Editor.h"
-#	include "LevelEditorViewport.h"
+#include "LevelEditorViewport.h"
 #endif //WITH_EDITOR
 
 
diff --git a/Engine/Source/Developer/GameplayDebugger/Private/GameplayDebuggerPrivate.h b/Engine/Source/Developer/GameplayDebugger/Private/GameplayDebuggerPrivate.h
--- a/Engine/Source/Developer/GameplayDebugger/Private/GameplayDebuggerPrivate.h
+++ b/Engine/Source/Developer/GameplayDebugger/Private/GameplayDebuggerPrivate.h
@@ -36,4 +36,8 @@
 #include "GameplayDebuggingControllerComponent.h"
 #include "GameplayDebuggingComponent.h"
 
+#if WITH_EDITOR
+#	include "Editor.h" //right now it's needed here because of ECoordSystem 
+#endif
+
 DECLARE_LOG_CATEGORY_EXTERN(LogGDT, Warning, All);
