CIS fixes.

diff --git a/Engine/Source/Developer/GameplayDebugger/Private/GameplayDebuggingHUDComponent.cpp b/Engine/Source/Developer/GameplayDebugger/Private/GameplayDebuggingHUDComponent.cpp
--- a/Engine/Source/Developer/GameplayDebugger/Private/GameplayDebuggingHUDComponent.cpp
+++ b/Engine/Source/Developer/GameplayDebugger/Private/GameplayDebuggingHUDComponent.cpp
@@ -14,6 +14,7 @@
 #include "Regex.h"
 #include "DrawDebugHelpers.h"
 #include "TimerManager.h"					// Game play timers
+#include "RenderUtils.h"
 
 DEFINE_LOG_CATEGORY_STATIC(LogHUD, Log, All);
 
diff --git a/Engine/Source/Runtime/Engine/Public/DynamicMeshBuilder.h b/Engine/Source/Runtime/Engine/Public/DynamicMeshBuilder.h
--- a/Engine/Source/Runtime/Engine/Public/DynamicMeshBuilder.h
+++ b/Engine/Source/Runtime/Engine/Public/DynamicMeshBuilder.h
@@ -5,6 +5,7 @@
 =============================================================================*/
 
 #pragma once
+#include "RenderUtils.h"
 
 class FPrimitiveDrawInterface;
 class FMaterialRenderProxy;
