Optimizations to GameplayDebugger, mostly for data replication. #ue4

diff --git a/Engine/Source/Developer/GameplayDebugger/Classes/GameplayDebuggingComponent.h b/Engine/Source/Developer/GameplayDebugger/Classes/GameplayDebuggingComponent.h
--- a/Engine/Source/Developer/GameplayDebugger/Classes/GameplayDebuggingComponent.h
+++ b/Engine/Source/Developer/GameplayDebugger/Classes/GameplayDebuggingComponent.h
@@ -90,10 +90,6 @@ class GAMEPLAYDEBUGGER_API UGameplayDebuggingComponent : public UPrimitiveCompon
 	TArray<uint8> NavmeshRepData;
 	
 	/** Begin EQS replication data */
-
-	UPROPERTY(Replicated)
-	TArray<FString> AllEQSName;
-
 	UPROPERTY(ReplicatedUsing = OnRep_UpdateEQS)
 	TArray<uint8> EQSRepData;
 	
diff --git a/Engine/Source/Developer/GameplayDebugger/Private/GameplayDebuggingComponent.cpp b/Engine/Source/Developer/GameplayDebugger/Private/GameplayDebuggingComponent.cpp
--- a/Engine/Source/Developer/GameplayDebugger/Private/GameplayDebuggingComponent.cpp
+++ b/Engine/Source/Developer/GameplayDebugger/Private/GameplayDebuggingComponent.cpp
@@ -132,7 +132,7 @@ UGameplayDebuggingComponent::UGameplayDebuggingComponent(const class FPostConstr
 		
 	PrimaryComponentTick.bCanEverTick = true;
 	bWantsInitializeComponent = true;
-	bAutoActivate = true;
+	bAutoActivate = false;
 	PrimaryComponentTick.bStartWithTickEnabled = false;
 
 	bIsSelectedForDebugging = false;
@@ -161,13 +161,9 @@ UGameplayDebuggingComponent::UGameplayDebuggingComponent(const class FPostConstr
 void UGameplayDebuggingComponent::Activate(bool bReset)
 {
 #if !(UE_BUILD_SHIPPING || UE_BUILD_TEST)
-	ActivationCounter++;
-	if (ActivationCounter > 0)
+	if (IsActive() == false)
 	{
-		if (IsActive() == false)
-		{
-			Super::Activate(bReset);
-		}
+		Super::Activate(bReset);
 		SetComponentTickEnabled(true);
 	}
 #else
@@ -178,14 +174,9 @@ void UGameplayDebuggingComponent::Activate(bool bReset)
 void UGameplayDebuggingComponent::Deactivate()
 {
 #if !(UE_BUILD_SHIPPING || UE_BUILD_TEST)
-	ActivationCounter--;
-	if (ActivationCounter <= 0)
+	if (IsActive())
 	{
-		ActivationCounter = 0;
-		if (IsActive())
-		{
-			Super::Deactivate();
-		}
+		Super::Deactivate();
 		SetComponentTickEnabled(false);
 	}
 #else
@@ -218,8 +209,6 @@ void UGameplayDebuggingComponent::GetLifetimeReplicatedProps( TArray< FLifetimeP
 	DOREPLIFETIME( UGameplayDebuggingComponent, TargetActor );
 
 	DOREPLIFETIME(UGameplayDebuggingComponent, EQSRepData);
-	DOREPLIFETIME(UGameplayDebuggingComponent, AllEQSName);
-
 #endif //!(UE_BUILD_SHIPPING || UE_BUILD_TEST)
 }
 
@@ -269,7 +258,7 @@ void UGameplayDebuggingComponent::TickComponent(float DeltaTime, enum ELevelTick
 
 	AGameplayDebuggingReplicator* Replicator = Cast<AGameplayDebuggingReplicator>(GetOwner());
 	UGameplayDebuggingComponent* DebugComponent = Replicator ? Replicator->GetDebugComponent() : NULL;
-	CurrentEQSIndex = DebugComponent ? FMath::Clamp(CurrentEQSIndex, 0, DebugComponent->AllEQSName.Num() - 1) : INDEX_NONE;
+	CurrentEQSIndex = DebugComponent ? FMath::Clamp(CurrentEQSIndex, 0, DebugComponent->EQSLocalData.Num() - 1) : INDEX_NONE;
 
 #endif //!(UE_BUILD_SHIPPING || UE_BUILD_TEST)
 }
@@ -432,7 +421,7 @@ void UGameplayDebuggingComponent::CollectBehaviorTreeData()
 	if (AAIController *MyController = Cast<AAIController>(MyPawn->Controller))
 	{
 		BrainComponentName = MyController->BrainComponent != NULL ? MyController->BrainComponent->GetName() : TEXT("");
-		BrainComponentString = MyController->BrainComponent->GetDebugInfoString();
+		BrainComponentString = MyController->BrainComponent != NULL ? MyController->BrainComponent->GetDebugInfoString() : TEXT("");
 		if (World && World->GetNetMode() != NM_Standalone)
 		{
 			BrainComponentString = BrainComponentString.Left(MAX_STRING_SERIALIZE_SIZE - 1);
@@ -586,14 +575,7 @@ void UGameplayDebuggingComponent::ServerReplicateData(uint32 InMessage, uint32
 		break;
 
 	case EDebugComponentMessage::DeactivateReplilcation:
-		{
-			Deactivate();
-			APawn* MyPawn = Cast<APawn>(GetSelectedActor());
-			if (MyPawn != NULL && IsActive() == false)
-			{
-				//MyPawn->RemoveDebugComponent();
-			}
-		}
+		Deactivate();
 		break;
 
 	case EDebugComponentMessage::ActivateDataView:
@@ -623,7 +605,7 @@ void UGameplayDebuggingComponent::OnChangeEQSQuery()
 {
 	AGameplayDebuggingReplicator* Replicator = Cast<AGameplayDebuggingReplicator>(GetOwner());
 	UGameplayDebuggingComponent* DebugComponent = Replicator ? Replicator->GetDebugComponent() : NULL;
-	if (DebugComponent && ++CurrentEQSIndex >= DebugComponent->AllEQSName.Num())
+	if (DebugComponent && ++CurrentEQSIndex >= DebugComponent->EQSLocalData.Num())
 	{
 		CurrentEQSIndex = 0;
 	}
@@ -687,13 +669,11 @@ void UGameplayDebuggingComponent::CollectEQSData()
 		return;
 	}
 
-	AllEQSName.Empty();
 	auto AllQueries = QueryManager->GetDebugger().GetAllQueriesForOwner(Owner);
 	for (int32 Index = 0; Index < AllQueries.Num(); ++Index)
 	{
 		EQSDebug::FQueryData* CurrentLocalData = NULL;
 		CachedQueryInstance = AllQueries[Index].Instance;
-		AllEQSName.AddUnique(CachedQueryInstance->QueryName);
 
 		 //find corresponding query
 		bool bSkipToNext = false;
@@ -1145,6 +1125,10 @@ void UGameplayDebuggingComponent::PrepareNavMeshData(struct FNavMeshSceneProxyDa
 //----------------------------------------------------------------------//
 FPrimitiveSceneProxy* UGameplayDebuggingComponent::CreateSceneProxy()
 {
+	if (!IsActive())
+	{
+		return NULL;
+	}
 	FDebugRenderSceneCompositeProxy* CompositeProxy = NULL;
 
 #if WITH_RECAST	
diff --git a/Engine/Source/Developer/GameplayDebugger/Private/GameplayDebuggingControllerComponent.cpp b/Engine/Source/Developer/GameplayDebugger/Private/GameplayDebuggingControllerComponent.cpp
--- a/Engine/Source/Developer/GameplayDebugger/Private/GameplayDebuggingControllerComponent.cpp
+++ b/Engine/Source/Developer/GameplayDebugger/Private/GameplayDebuggingControllerComponent.cpp
@@ -154,6 +154,7 @@ void UGameplayDebuggingControllerComponent::OnActivationKeyPressed()
 
 			BindAIDebugViewKeys();
 			GetDebuggingReplicator()->EnableDraw(true);
+			GetDebuggingReplicator()->ServerReplicateMessage(NULL, EDebugComponentMessage::ActivateReplication, EAIDebugDrawDataView::Empty);
 		}
 
 		ControlKeyPressedTime = GetWorld()->GetTimeSeconds();
@@ -199,6 +200,7 @@ void UGameplayDebuggingControllerComponent::CloseDebugTool()
 		Deactivate();
 		SetComponentTickEnabled(false);
 		GetDebuggingReplicator()->EnableDraw(false);
+		GetDebuggingReplicator()->ServerReplicateMessage(NULL, EDebugComponentMessage::DeactivateReplilcation, EAIDebugDrawDataView::Empty);
 		bToolActivated = false;
 	}
 }
@@ -262,7 +264,6 @@ void UGameplayDebuggingControllerComponent::ToggleAIDebugView_SetView0()
 				UpdateNavMeshTimer();
 
 				GetDebuggingReplicator()->ServerReplicateMessage(Pawn, EDebugComponentMessage::ActivateDataView, EAIDebugDrawDataView::NavMesh);
-				GetDebuggingReplicator()->ServerReplicateMessage(Pawn, EDebugComponentMessage::ActivateReplication, EAIDebugDrawDataView::Empty);
 				OwnerComp->SetVisibility(true, true);
 				OwnerComp->SetHiddenInGame(false, true);
 				OwnerComp->MarkRenderStateDirty();
@@ -272,7 +273,6 @@ void UGameplayDebuggingControllerComponent::ToggleAIDebugView_SetView0()
 				GetWorld()->GetTimerManager().ClearTimer(this, &UGameplayDebuggingControllerComponent::UpdateNavMeshTimer);
 
 				GetDebuggingReplicator()->ServerReplicateMessage(Pawn, EDebugComponentMessage::DeactivateDataView, EAIDebugDrawDataView::NavMesh);
-				GetDebuggingReplicator()->ServerReplicateMessage(Pawn, EDebugComponentMessage::DeactivateReplilcation, EAIDebugDrawDataView::Empty);
 				OwnerComp->ServerDiscardNavmeshData();
 				OwnerComp->SetVisibility(false, true);
 				OwnerComp->SetHiddenInGame(true, true);
diff --git a/Engine/Source/Developer/GameplayDebugger/Private/GameplayDebuggingHUDComponent.cpp b/Engine/Source/Developer/GameplayDebugger/Private/GameplayDebuggingHUDComponent.cpp
--- a/Engine/Source/Developer/GameplayDebugger/Private/GameplayDebuggingHUDComponent.cpp
+++ b/Engine/Source/Developer/GameplayDebugger/Private/GameplayDebuggingHUDComponent.cpp
@@ -10,6 +10,7 @@
 #include "GameplayDebuggingControllerComponent.h"
 #include "CanvasItem.h"
 #include "AI/Navigation/NavigationSystem.h"
+#include "EnvironmentQuery/EnvQueryTypes.h"
 #include "Engine/Texture2D.h"
 #include "Regex.h"
 #include "DrawDebugHelpers.h"
@@ -354,27 +355,27 @@ void AGameplayDebuggingHUDComponent::DrawEQSData(APlayerController* PC, class UG
 #if !(UE_BUILD_SHIPPING || UE_BUILD_TEST) && WITH_EQS
 	PrintString(DefaultContext, TEXT("\n{green}EQS {white}[Use + key to switch query]\n"));
 
-	if (DebugComponent->AllEQSName.Num() == 0)
+	if (DebugComponent->EQSLocalData.Num() == 0)
 	{
 		return;
 	}
 
-	if (!DebugComponent->AllEQSName.IsValidIndex(DebugComponent->CurrentEQSIndex))
+	if (!DebugComponent->EQSLocalData.IsValidIndex(DebugComponent->CurrentEQSIndex))
 	{
 		return;
 	}
-	const FString SelecterEQSName = DebugComponent->AllEQSName[DebugComponent->CurrentEQSIndex];
+	const FString SelecterEQSName = DebugComponent->EQSLocalData[DebugComponent->CurrentEQSIndex].Name;
 
 	PrintString(DefaultContext, TEXT("{white}Queries: "));
-	for (auto CurrentQueryName : DebugComponent->AllEQSName)
+	for (auto CurrentQuery : DebugComponent->EQSLocalData)
 	{
-		if (SelecterEQSName == CurrentQueryName)
+		if (SelecterEQSName == CurrentQuery.Name)
 		{
-			PrintString(DefaultContext, FString::Printf(TEXT("{green}%s, "), *CurrentQueryName));
+			PrintString(DefaultContext, FString::Printf(TEXT("{green}%s, "), *CurrentQuery.Name));
 		}
 		else
 		{
-			PrintString(DefaultContext, FString::Printf(TEXT("{yellow}%s, "), *CurrentQueryName));
+			PrintString(DefaultContext, FString::Printf(TEXT("{yellow}%s, "), *CurrentQuery.Name));
 		}
 	}
 	PrintString(DefaultContext, TEXT("\n"));
@@ -389,7 +390,6 @@ void AGameplayDebuggingHUDComponent::DrawEQSData(APlayerController* PC, class UG
 			, CurrentLocalData.Id
 			));
 
-#if  USE_EQS_DEBUGGER
 		// draw test weights for best X items
 		const int32 NumItems = CurrentLocalData.Items.Num();
 		const int32 NumTests = CurrentLocalData.Tests.Num();
@@ -443,14 +443,13 @@ void AGameplayDebuggingHUDComponent::DrawEQSData(APlayerController* PC, class UG
 
 			PrintString(DefaultContext, TestDesc);
 		}
-#endif //USE_EQS_DEBUGGER
 	}
 #endif //!(UE_BUILD_SHIPPING || UE_BUILD_TEST)
 }
 
 void AGameplayDebuggingHUDComponent::DrawEQSItemDetails(int32 ItemIdx, class UGameplayDebuggingComponent *DebugComponent)
 {
-#if USE_EQS_DEBUGGER && WITH_EQS
+#if !(UE_BUILD_SHIPPING || UE_BUILD_TEST) && WITH_EQS
 	const float PosY = DefaultContext.CursorY + 1.0f;
 	float PosX = DefaultContext.CursorX;
 
@@ -496,7 +495,7 @@ void AGameplayDebuggingHUDComponent::DrawEQSItemDetails(int32 ItemIdx, class UGa
 		PrintString(DefaultContext, FColor::Green, TestDesc, PosX, PosY);
 		PosX += 100.0f;
 	}
-#endif
+#endif //!(UE_BUILD_SHIPPING || UE_BUILD_TEST) && WITH_EQS
 }
 
 void AGameplayDebuggingHUDComponent::DrawPerception(APlayerController* PC, class UGameplayDebuggingComponent *DebugComponent)
diff --git a/Engine/Source/Developer/GameplayDebugger/Private/GameplayDebuggingReplicator.cpp b/Engine/Source/Developer/GameplayDebugger/Private/GameplayDebuggingReplicator.cpp
--- a/Engine/Source/Developer/GameplayDebugger/Private/GameplayDebuggingReplicator.cpp
+++ b/Engine/Source/Developer/GameplayDebugger/Private/GameplayDebuggingReplicator.cpp
@@ -23,7 +23,7 @@ AGameplayDebuggingReplicator::AGameplayDebuggingReplicator(const class FPostCons
 	, PlayerControllersUpdateDelay(0)
 {
 	PrimaryActorTick.bCanEverTick = true;
-	PrimaryActorTick.bStartWithTickEnabled = true;
+	PrimaryActorTick.bStartWithTickEnabled = false;
 	
 	TSubobjectPtr<USceneComponent> SceneComponent = PCIP.CreateDefaultSubobject<USceneComponent>(this, TEXT("SceneComponent"));
 	RootComponent = SceneComponent;
@@ -123,7 +123,7 @@ void AGameplayDebuggingReplicator::BeginPlay()
 			DebugComponent = ConstructObject<UGameplayDebuggingComponent>(DebugComponentClass.Get(), this);
 			DebugComponent->SetIsReplicated(true);
 			DebugComponent->RegisterComponent();
-			DebugComponent->Activate();
+			//DebugComponent->Activate();
 		}
 	}
 
@@ -154,6 +154,9 @@ void AGameplayDebuggingReplicator::BeginPlay()
 			DebugComponentClass = UGameplayDebuggingComponent::StaticClass();
 		}
 	}
+
+	SetActorTickEnabled(true);
+	SetTickableWhenPaused(true);
 }
 
 UGameplayDebuggingComponent* AGameplayDebuggingReplicator::GetDebugComponent()
@@ -231,6 +234,7 @@ void AGameplayDebuggingReplicator::EnableDraw(bool bEnable)
 		GameHUD->bShowHUD = bEnable ? false : true;
 	}
 	GEngine->bEnableOnScreenDebugMessages = bEnable ? false : true;
+	GetDebugComponent()->MarkRenderStateDirty();
 }
 
 bool AGameplayDebuggingReplicator::IsToolCreated()
@@ -402,7 +406,7 @@ void AGameplayDebuggingReplicator::DrawDebugDataDelegate(class UCanvas* Canvas,
 		return;
 	}
 
-	if (!LocalPlayerOwner || IsGlobalInWorld())
+	if (!LocalPlayerOwner || IsGlobalInWorld() || !IsDrawEnabled())
 	{
 		return;
 	}
@@ -419,7 +423,7 @@ void AGameplayDebuggingReplicator::DrawDebugDataDelegate(class UCanvas* Canvas,
 	LastDrawAtFrame = GFrameNumber;
 
 	const UGameplayDebuggingControllerComponent*  GDC = FindComponentByClass<UGameplayDebuggingControllerComponent>();
-	if (!IsDrawEnabled() || !GDC)
+	if (!GDC)
 	{
 		return;
 	}
diff --git a/Engine/Source/Runtime/AIModule/Private/EnvironmentQuery/EnvQueryManager.cpp b/Engine/Source/Runtime/AIModule/Private/EnvironmentQuery/EnvQueryManager.cpp
--- a/Engine/Source/Runtime/AIModule/Private/EnvironmentQuery/EnvQueryManager.cpp
+++ b/Engine/Source/Runtime/AIModule/Private/EnvironmentQuery/EnvQueryManager.cpp
@@ -519,7 +519,13 @@ void UEnvQueryManager::CreateOptionInstance(UEnvQueryOption* OptionTemplate, con
 
 void FEQSDebugger::StoreQuery(TSharedPtr<FEnvQueryInstance>& Query)
 {
+	StoredQueries.Remove(NULL);
+	if (!Query.IsValid())
+	{
+		return;
+	}
 	TArray<FEnvQueryInfo>& AllQueries = StoredQueries.FindOrAdd(Query->Owner.Get());
+		
 	bool bFoundQuery = false;
 	for (auto It = AllQueries.CreateIterator(); It; ++It)
 	{
