- Extended GameplayDebugger data view in Simulation. I added a way to switch debug views, to have all functionality from PIE.
- Fixed paths rendering while using GameplayDebugger (client/server too)
Both changes integrated to Main from Dev (CL #2290675, CL #2282927)

diff --git a/Engine/Source/Developer/GameplayDebugger/Classes/GameplayDebuggingComponent.h b/Engine/Source/Developer/GameplayDebugger/Classes/GameplayDebuggingComponent.h
--- a/Engine/Source/Developer/GameplayDebugger/Classes/GameplayDebuggingComponent.h
+++ b/Engine/Source/Developer/GameplayDebugger/Classes/GameplayDebuggingComponent.h
@@ -43,12 +43,6 @@ class GAMEPLAYDEBUGGER_API UGameplayDebuggingComponent : public UPrimitiveCompon
 	UPROPERTY(globalconfig)
 	FString DebugComponentClassName;
 
-	UPROPERTY(Replicated)
-	bool bIsSelectedForDebugging;
-
-	UPROPERTY(Replicated)
-	int32 ActivationCounter;
-
 	UPROPERTY(Replicated)
 	int32 ShowExtendedInformatiomCounter;
 
@@ -90,6 +84,10 @@ class GAMEPLAYDEBUGGER_API UGameplayDebuggingComponent : public UPrimitiveCompon
 	TArray<uint8> NavmeshRepData;
 	
 	/** Begin EQS replication data */
+
+	UPROPERTY(Replicated)
+	TArray<FString> AllEQSName;
+
 	UPROPERTY(ReplicatedUsing = OnRep_UpdateEQS)
 	TArray<uint8> EQSRepData;
 	
@@ -134,10 +132,6 @@ class GAMEPLAYDEBUGGER_API UGameplayDebuggingComponent : public UPrimitiveCompon
 
 	virtual void EnableDebugDraw(bool bEnable, bool InFocusedComponent = false);
 
-	FORCEINLINE bool IsSelected() const { return bIsSelectedForDebugging; }
-	/** Will broadcast information that this component is (no longer) being "observed" */
-	void SelectForDebugging(bool bNewStatus);
-
 	bool ShouldReplicateData(EAIDebugDrawDataView::Type InView) const { return ReplicateViewDataCounters[InView] > 0 /*true*/; }
 	virtual void CollectDataToReplicate(bool bCollectExtendedData);
 
@@ -156,6 +150,8 @@ class GAMEPLAYDEBUGGER_API UGameplayDebuggingComponent : public UPrimitiveCompon
 		return TargetActor;
 	}
 
+	void SetEQSIndex(int32 Index) { CurrentEQSIndex = AllEQSName.Num() > 0 ? FMath::Clamp(Index, 0, AllEQSName.Num() - 1) : INDEX_NONE; }
+	int32 GetEQSIndex() { return CurrentEQSIndex; }
 	//=============================================================================
 	// EQS debugging
 	//=============================================================================
diff --git a/Engine/Source/Developer/GameplayDebugger/Classes/GameplayDebuggingReplicator.h b/Engine/Source/Developer/GameplayDebugger/Classes/GameplayDebuggingReplicator.h
--- a/Engine/Source/Developer/GameplayDebugger/Classes/GameplayDebuggingReplicator.h
+++ b/Engine/Source/Developer/GameplayDebugger/Classes/GameplayDebuggingReplicator.h
@@ -29,9 +29,48 @@ class GAMEPLAYDEBUGGER_API AGameplayDebuggingReplicator : public AActor
 	UPROPERTY(Replicated, Transient)
 	APlayerController* LocalPlayerOwner;
 
+	UPROPERTY(Replicated, Transient)
+	AActor*	LastSelectedActorToDebug;
+
 	UPROPERTY(Replicated, Transient)
 	bool bIsGlobalInWorld;
 
+	UPROPERTY(Transient, EditAnywhere, Category = "DataView")
+	bool OverHead;
+	
+	UPROPERTY(Transient, EditAnywhere, Category = "DataView")
+	bool Basic;
+	
+	UPROPERTY(Transient, EditAnywhere, Category = "DataView")
+	bool BehaviorTree;
+	
+	UPROPERTY(Transient, EditAnywhere, Category = "DataView|EQS")
+	bool EQS;
+
+	UPROPERTY(Transient, EditAnywhere, Category = "DataView|EQS", meta = (EditCondition = "EQS"))
+	bool EnableEQSOnHUD;
+
+	UPROPERTY(Transient, EditAnywhere, Category = "DataView|EQS", meta = (EditCondition = "EQS"))
+	int32 ActiveEQSIndex;
+
+	UPROPERTY(Transient, EditAnywhere, Category = "DataView")
+	bool Perception;
+	
+	UPROPERTY(Transient, EditAnywhere, Category = "DataView")
+	bool GameView1;
+	
+	UPROPERTY(Transient, EditAnywhere, Category = "DataView")
+	bool GameView2;
+	
+	UPROPERTY(Transient, EditAnywhere, Category = "DataView")
+	bool GameView3;
+	
+	UPROPERTY(Transient, EditAnywhere, Category = "DataView")
+	bool GameView4;
+	
+	UPROPERTY(Transient, EditAnywhere, Category = DataView)
+	bool GameView5;
+
 	UFUNCTION(reliable, server, WithValidation)
 	void ServerReplicateMessage(class AActor* Actor, uint32 InMessage, uint32 DataView = 0);
 
@@ -41,6 +80,10 @@ class GAMEPLAYDEBUGGER_API AGameplayDebuggingReplicator : public AActor
 	UFUNCTION(reliable, server, WithValidation)
 	void ServerEnableTargetSelection(bool bEnable, APlayerController* Context);
 
+#if WITH_EDITOR
+	void PostEditChangeProperty(FPropertyChangedEvent& PropertyChangedEvent) override;
+#endif
+
 	virtual class UNetConnection* GetNetConnection() override;
 
 	virtual bool IsNetRelevantFor(class APlayerController* RealViewer, AActor* Viewer, const FVector& SrcLocation) override;
@@ -64,7 +107,7 @@ class GAMEPLAYDEBUGGER_API AGameplayDebuggingReplicator : public AActor
 	void SetLocalPlayerOwner(APlayerController* PC) { LocalPlayerOwner = PC; }
 	APlayerController* GetLocalPlayerOwner() { return LocalPlayerOwner; }
 
-	FORCEINLINE AActor* GetSelectedActorToDebug() { return LastSelectedActorInSimulation.Get(); }
+	FORCEINLINE AActor* GetSelectedActorToDebug() { return LastSelectedActorToDebug; }
 	void SetActorToDebug(AActor* InActor);
 
 	uint32 DebuggerShowFlags;
@@ -85,5 +128,4 @@ private:
 	TWeakObjectPtr<UClass> DebugComponentHUDClass;
 
 	TWeakObjectPtr<AGameplayDebuggingHUDComponent>	DebugRenderer;
-	TWeakObjectPtr<AActor>	LastSelectedActorInSimulation;
 };
diff --git a/Engine/Source/Developer/GameplayDebugger/Classes/GameplayDebuggingTypes.h b/Engine/Source/Developer/GameplayDebugger/Classes/GameplayDebuggingTypes.h
--- a/Engine/Source/Developer/GameplayDebugger/Classes/GameplayDebuggingTypes.h
+++ b/Engine/Source/Developer/GameplayDebugger/Classes/GameplayDebuggingTypes.h
@@ -21,7 +21,6 @@ namespace EAIDebugDrawDataView
 		GameView4,
 		GameView5,
 		NavMesh,
-		EditorDebugAIFlag,
 		MAX UMETA(Hidden)
 	};
 }
diff --git a/Engine/Source/Developer/GameplayDebugger/Private/GameplayDebuggerExec.cpp b/Engine/Source/Developer/GameplayDebugger/Private/GameplayDebuggerExec.cpp
--- a/Engine/Source/Developer/GameplayDebugger/Private/GameplayDebuggerExec.cpp
+++ b/Engine/Source/Developer/GameplayDebugger/Private/GameplayDebuggerExec.cpp
@@ -99,69 +99,6 @@ bool FGameplayDebuggerExec::Exec(UWorld* Inworld, const TCHAR* Cmd, FOutputDevic
 			}
 		}
 	}
-	else if (FParse::Command(&Cmd, TEXT("ToggleGameplayDebugView")))
-	{
-		static TArray<FString> ViewNames;
-		if (ViewNames.Num() == 0)
-		{
-			const UEnum* ViewlEnum = FindObject<UEnum>(ANY_PACKAGE, TEXT("EAIDebugDrawDataView"), true);
-			ViewNames.AddZeroed(EAIDebugDrawDataView::MAX);
-			for (int32 Index = 0; Index < EAIDebugDrawDataView::MAX; ++Index)
-			{
-				ViewNames[Index] = ViewlEnum->GetEnumName(Index);
-			}
-		}
-
-		const bool bActivePIE = Inworld && Inworld->IsGameWorld();
-		AGameplayDebuggingReplicator* DebuggingReplicator = NULL;
-		if (bActivePIE)
-		{
-			for (FActorIterator It(Inworld); It; ++It)
-			{
-				AActor* A = *It;
-				if (A && A->IsA(AGameplayDebuggingReplicator::StaticClass()) && !A->IsPendingKill())
-				{
-					AGameplayDebuggingReplicator* Replicator = Cast<AGameplayDebuggingReplicator>(A);
-					if (Replicator->GetLocalPlayerOwner() == PC)
-					{
-						DebuggingReplicator = Replicator;
-						break;
-					}
-				}
-			}
-		}
-
-		FGameplayDebuggerSettings DebuggerSettings = GameplayDebuggerSettings(DebuggingReplicator);
-		FString InViewName = FParse::Token(Cmd, 0);
-		int32 ViewIndex = ViewNames.Find(InViewName);
-		FGameplayDebuggerSettings::ShowFlagIndex = FEngineShowFlags::FindIndexByName(TEXT("GameplayDebug"));
-		if (ViewIndex != INDEX_NONE)
-		{
-			DebuggerSettings.CheckFlag(ViewIndex) ? DebuggerSettings.ClearFlag(ViewIndex) : DebuggerSettings.SetFlag(ViewIndex);
-			if (bActivePIE)
-			{
-				GameplayDebuggerSettings().CheckFlag(ViewIndex) ? GameplayDebuggerSettings().ClearFlag(ViewIndex) : GameplayDebuggerSettings().SetFlag(ViewIndex);
-			}
-#if WITH_EDITOR
-			if (ViewIndex == EAIDebugDrawDataView::EQS && GCurrentLevelEditingViewportClient)
-			{
-				GCurrentLevelEditingViewportClient->EngineShowFlags.SetSingleFlag(FGameplayDebuggerSettings::ShowFlagIndex, DebuggerSettings.CheckFlag(ViewIndex));
-			}
-#endif
-			PC->ClientMessage(FString::Printf(TEXT("View %s %s")
-				, *InViewName
-				, DebuggerSettings.CheckFlag(ViewIndex) ? TEXT("enabled") : TEXT("disabled")));
-		}
-		else
-		{
-			PC->ClientMessage(TEXT("Unknown debug view name. Valid options are:"));
-			for (int32 Index = 0; Index < EAIDebugDrawDataView::MAX; ++Index)
-			{
-				PC->ClientMessage(*ViewNames[Index]);
-			}
-		}
-		bHandled = true;
-	}
 	else if (FParse::Command(&Cmd, TEXT("RunEQS")))
 	{
 		bHandled = true;
diff --git a/Engine/Source/Developer/GameplayDebugger/Private/GameplayDebuggingComponent.cpp b/Engine/Source/Developer/GameplayDebugger/Private/GameplayDebuggingComponent.cpp
--- a/Engine/Source/Developer/GameplayDebugger/Private/GameplayDebuggingComponent.cpp
+++ b/Engine/Source/Developer/GameplayDebugger/Private/GameplayDebuggingComponent.cpp
@@ -136,8 +136,6 @@ UGameplayDebuggingComponent::UGameplayDebuggingComponent(const class FPostConstr
 	bAutoActivate = false;
 	PrimaryComponentTick.bStartWithTickEnabled = false;
 
-	bIsSelectedForDebugging = false;
-	ActivationCounter = 0;
 	ShowExtendedInformatiomCounter = 0;
 #if WITH_EDITOR
 	bWasSelectedInEditor = false;
@@ -162,11 +160,9 @@ UGameplayDebuggingComponent::UGameplayDebuggingComponent(const class FPostConstr
 void UGameplayDebuggingComponent::Activate(bool bReset)
 {
 #if !(UE_BUILD_SHIPPING || UE_BUILD_TEST)
-		if (IsActive() == false)
-		{
-			Super::Activate(bReset);
+		Super::Activate(bReset);
 		SetComponentTickEnabled(true);
-	}
+		SetIsReplicated(true);
 #else
 	Super::Activate(bReset);
 #endif
@@ -175,11 +171,9 @@ void UGameplayDebuggingComponent::Activate(bool bReset)
 void UGameplayDebuggingComponent::Deactivate()
 {
 #if !(UE_BUILD_SHIPPING || UE_BUILD_TEST)
-		if (IsActive())
-		{
-			Super::Deactivate();
+		Super::Deactivate();
 		SetComponentTickEnabled(false);
-	}
+		SetIsReplicated(false);
 #else
 	Super::Deactivate();
 #endif
@@ -189,8 +183,6 @@ void UGameplayDebuggingComponent::GetLifetimeReplicatedProps( TArray< FLifetimeP
 {
 	Super::GetLifetimeReplicatedProps( OutLifetimeProps );
 #if !(UE_BUILD_SHIPPING || UE_BUILD_TEST)
-	DOREPLIFETIME(UGameplayDebuggingComponent, bIsSelectedForDebugging);
-	DOREPLIFETIME( UGameplayDebuggingComponent, ActivationCounter );
 	DOREPLIFETIME( UGameplayDebuggingComponent, ReplicateViewDataCounters );
 	DOREPLIFETIME( UGameplayDebuggingComponent, ShowExtendedInformatiomCounter );
 	DOREPLIFETIME( UGameplayDebuggingComponent, ControllerName )
@@ -233,6 +225,9 @@ void UGameplayDebuggingComponent::SetActorToDebug(AActor* Actor)
 {
 #if !(UE_BUILD_SHIPPING || UE_BUILD_TEST)
 	TargetActor = Actor;
+	EQSLocalData.Reset();
+	AllEQSName.Reset();
+	CurrentEQSIndex = INDEX_NONE;
 #endif //!(UE_BUILD_SHIPPING || UE_BUILD_TEST)
 }
 
@@ -258,8 +253,7 @@ void UGameplayDebuggingComponent::TickComponent(float DeltaTime, enum ELevelTick
 	}
 
 	AGameplayDebuggingReplicator* Replicator = Cast<AGameplayDebuggingReplicator>(GetOwner());
-	UGameplayDebuggingComponent* DebugComponent = Replicator ? Replicator->GetDebugComponent() : NULL;
-	CurrentEQSIndex = DebugComponent ? FMath::Clamp(CurrentEQSIndex, 0, DebugComponent->EQSLocalData.Num() - 1) : INDEX_NONE;
+	CurrentEQSIndex = AllEQSName.Num() > 0 ?FMath::Clamp(CurrentEQSIndex, 0, AllEQSName.Num() - 1) : INDEX_NONE;
 
 #endif //!(UE_BUILD_SHIPPING || UE_BUILD_TEST)
 }
@@ -332,7 +326,6 @@ void UGameplayDebuggingComponent::SelectTargetToDebug()
 
 			//always update component for best target
 			SetActorToDebug(Cast<AActor>(BestTarget));
-			SelectForDebugging(true);
 			ServerReplicateData(EDebugComponentMessage::ActivateReplication, EAIDebugDrawDataView::Empty);
 		}
 	}
@@ -352,12 +345,14 @@ void UGameplayDebuggingComponent::CollectDataToReplicate(bool bCollectExtendedDa
 		CollectBasicData();
 	}
 
-	if (IsSelected() && ShouldReplicateData(EAIDebugDrawDataView::EditorDebugAIFlag))
+	AGameplayDebuggingReplicator* Replicator = Cast<AGameplayDebuggingReplicator>(GetOwner());
+	const bool bDrawFullData = Replicator->GetSelectedActorToDebug() == GetSelectedActor();
+	if (bDrawFullData && ShouldReplicateData(EAIDebugDrawDataView::Basic))
 	{
 		CollectPathData();
 	}
 
-	if (bCollectExtendedData && IsSelected())
+	if (bCollectExtendedData && bDrawFullData)
 	{
 		if (ShouldReplicateData(EAIDebugDrawDataView::BehaviorTree))
 		{
@@ -506,50 +501,12 @@ void UGameplayDebuggingComponent::CollectPathData()
 #endif //!(UE_BUILD_SHIPPING || UE_BUILD_TEST)
 }
 
-void UGameplayDebuggingComponent::SelectForDebugging(bool bNewStatus)
-{
-#if !(UE_BUILD_SHIPPING || UE_BUILD_TEST)
-	if (bIsSelectedForDebugging != bNewStatus)
-	{
-		bIsSelectedForDebugging = bNewStatus;
-		
-#if 0
-		// temp: avoidance debug
-		UAvoidanceManager* Avoidance = GetWorld()->GetAvoidanceManager();
-		AController* ControllerOwner = Cast<AController>(GetOwner());
-		APawn* PawnOwner = ControllerOwner ? ControllerOwner->GetPawn() : Cast<APawn>(GetOwner());
-		UCharacterMovementComponent* MovementComp = PawnOwner ? PawnOwner->FindComponentByClass<UCharacterMovementComponent>() : NULL;
-		if (MovementComp && Avoidance)
-		{
-			Avoidance->AvoidanceDebugForUID(MovementComp->AvoidanceUID, bNewStatus);
-		}
-#endif
-
-#if WITH_EQS
-		if (bNewStatus == false)
-		{
-			CachedQueryInstance.Reset();
-			MarkRenderStateDirty();
-		}
-#endif // WITH_EQS
-	}
-#endif //!(UE_BUILD_SHIPPING || UE_BUILD_TEST)
-}
-
 void UGameplayDebuggingComponent::EnableDebugDraw(bool bEnable, bool InFocusedComponent)
 {
-#if !(UE_BUILD_SHIPPING || UE_BUILD_TEST)
-	if (bEnable)
-	{
-		SelectForDebugging(InFocusedComponent);
-	}
-	else
-	{
-		SelectForDebugging(false);
+#if !(UE_BUILD_SHIPPING || UE_BUILD_TEST) 
 #if WITH_EQS
-		EnableClientEQSSceneProxy(false);
+	EnableClientEQSSceneProxy(bEnable);
 #endif // WITH_EQS
-	}
 #endif //!(UE_BUILD_SHIPPING || UE_BUILD_TEST)
 }
 
@@ -573,7 +530,7 @@ void UGameplayDebuggingComponent::ServerReplicateData(uint32 InMessage, uint32
 		break;
 
 	case EDebugComponentMessage::DeactivateReplilcation:
-			Deactivate();
+		Deactivate();
 		break;
 
 	case EDebugComponentMessage::ActivateDataView:
@@ -602,8 +559,7 @@ void UGameplayDebuggingComponent::ServerReplicateData(uint32 InMessage, uint32
 void UGameplayDebuggingComponent::OnChangeEQSQuery()
 {
 	AGameplayDebuggingReplicator* Replicator = Cast<AGameplayDebuggingReplicator>(GetOwner());
-	UGameplayDebuggingComponent* DebugComponent = Replicator ? Replicator->GetDebugComponent() : NULL;
-	if (DebugComponent && ++CurrentEQSIndex >= DebugComponent->EQSLocalData.Num())
+	if (++CurrentEQSIndex >= AllEQSName.Num())
 	{
 		CurrentEQSIndex = 0;
 	}
@@ -642,9 +598,9 @@ void UGameplayDebuggingComponent::OnRep_UpdateEQS()
 		FMemoryReader ArReader(UncompressedBuffer);
 
 		ArReader << EQSLocalData;
-		CurrentEQSIndex = 0;
 	}	
 
+	CurrentEQSIndex = AllEQSName.Num() > 0 ? FMath::Clamp(CurrentEQSIndex, 0, AllEQSName.Num() - 1) : INDEX_NONE;
 	UpdateBounds();
 	MarkRenderStateDirty();
 #endif //USE_EQS_DEBUGGER
@@ -672,6 +628,8 @@ void UGameplayDebuggingComponent::CollectEQSData()
 	{
 		EQSDebug::FQueryData* CurrentLocalData = NULL;
 		CachedQueryInstance = AllQueries[Index].Instance;
+		float CachedTimestamp = AllQueries[Index].Timestamp;
+		AllEQSName.AddUnique(CachedQueryInstance->QueryName);
 
 		 //find corresponding query
 		bool bSkipToNext = false;
@@ -681,8 +639,11 @@ void UGameplayDebuggingComponent::CollectEQSData()
 			{
 				if (EQSLocalData[Idx].Id == CachedQueryInstance->QueryID)
 				{
-					bSkipToNext = true;
-					break;
+					if (EQSLocalData[Idx].Timestamp == CachedTimestamp)
+					{
+						bSkipToNext = true;
+						break;
+					}
 				}
 				CurrentLocalData = &EQSLocalData[Index];;
 				break;
@@ -1130,7 +1091,6 @@ FPrimitiveSceneProxy* UGameplayDebuggingComponent::CreateSceneProxy()
 	FDebugRenderSceneCompositeProxy* CompositeProxy = NULL;
 
 #if WITH_RECAST	
-	const APawn* MyPawn = Cast<APawn>(GetSelectedActor());
 	if (ShouldReplicateData(EAIDebugDrawDataView::NavMesh) && World && World->GetNetMode() != NM_DedicatedServer)
 	{
 		FNavMeshSceneProxyData NewNavmeshRenderData;
@@ -1146,13 +1106,14 @@ FPrimitiveSceneProxy* UGameplayDebuggingComponent::CreateSceneProxy()
 #endif
 
 #if USE_EQS_DEBUGGER
-	if (ShouldReplicateData(EAIDebugDrawDataView::EQS) && IsClientEQSSceneProxyEnabled() )
+	if (ShouldReplicateData(EAIDebugDrawDataView::EQS) && IsClientEQSSceneProxyEnabled())
 	{
 		if (EQSLocalData.IsValidIndex(CurrentEQSIndex))
 		{
 			CompositeProxy = CompositeProxy ? CompositeProxy : (new FDebugRenderSceneCompositeProxy(this));
 			auto& CurrentLocalData = EQSLocalData[CurrentEQSIndex];
-			CompositeProxy->AddChild(new FEQSSceneProxy(this, TEXT("GameplayDebug"), false, CurrentLocalData.SolidSpheres, CurrentLocalData.Texts));
+			bool bUseDebugAIFlag = false;
+			CompositeProxy->AddChild(new FEQSSceneProxy(this, TEXT("DebugAI"), false, CurrentLocalData.SolidSpheres, CurrentLocalData.Texts));
 		}
 	}
 #endif // USE_EQS_DEBUGGER
diff --git a/Engine/Source/Developer/GameplayDebugger/Private/GameplayDebuggingControllerComponent.cpp b/Engine/Source/Developer/GameplayDebugger/Private/GameplayDebuggingControllerComponent.cpp
--- a/Engine/Source/Developer/GameplayDebugger/Private/GameplayDebuggingControllerComponent.cpp
+++ b/Engine/Source/Developer/GameplayDebugger/Private/GameplayDebuggingControllerComponent.cpp
@@ -9,6 +9,7 @@
 
 DEFINE_LOG_CATEGORY_STATIC(LogGameplayDebugging, Log, All);
 
+#define USE_ALTERNATIVE_KEYS 0
 #define BUGIT_VIEWS (1<<EAIDebugDrawDataView::Basic) | (1 << EAIDebugDrawDataView::OverHead)
 #define BREAK_LINE_TEXT TEXT("________________________________________________________________")
 
@@ -155,6 +156,7 @@ void UGameplayDebuggingControllerComponent::OnActivationKeyPressed()
 			BindAIDebugViewKeys();
 			GetDebuggingReplicator()->EnableDraw(true);
 			GetDebuggingReplicator()->ServerReplicateMessage(NULL, EDebugComponentMessage::ActivateReplication, EAIDebugDrawDataView::Empty);
+			PlayerOwner->ConsoleCommand(TEXT("ShowFlag.DebugAI 1"), false);
 		}
 
 		ControlKeyPressedTime = GetWorld()->GetTimeSeconds();
@@ -199,9 +201,14 @@ void UGameplayDebuggingControllerComponent::CloseDebugTool()
 	{
 		Deactivate();
 		SetComponentTickEnabled(false);
+		GetDebuggingReplicator()->ServerReplicateMessage(NULL, EDebugComponentMessage::DeactivateReplilcation, EAIDebugDrawDataView::Empty);
 		GetDebuggingReplicator()->EnableDraw(false);
 		GetDebuggingReplicator()->ServerReplicateMessage(NULL, EDebugComponentMessage::DeactivateReplilcation, EAIDebugDrawDataView::Empty);
 		bToolActivated = false;
+		if (PlayerOwner.IsValid())
+		{
+			PlayerOwner->ConsoleCommand(TEXT("ShowFlag.DebugAI 0"), false);
+		}
 	}
 }
 
@@ -229,6 +236,19 @@ void UGameplayDebuggingControllerComponent::BindAIDebugViewKeys()
 		AIDebugViewInputComponent->BindKey(EKeys::NumPadEight, IE_Pressed, this, &UGameplayDebuggingControllerComponent::ToggleAIDebugView_SetView8);
 		AIDebugViewInputComponent->BindKey(EKeys::NumPadNine, IE_Pressed, this, &UGameplayDebuggingControllerComponent::ToggleAIDebugView_SetView9);
 		AIDebugViewInputComponent->BindKey(EKeys::Add, IE_Released, this, &UGameplayDebuggingControllerComponent::NextEQSQuery);
+#if USE_ALTERNATIVE_KEYS
+		AIDebugViewInputComponent->BindKey(FInputChord(EKeys::Zero, false, false, true, false), IE_Pressed, this, &UGameplayDebuggingControllerComponent::ToggleAIDebugView_SetView0);
+		AIDebugViewInputComponent->BindKey(FInputChord(EKeys::One, false, false, true, false), IE_Pressed, this, &UGameplayDebuggingControllerComponent::ToggleAIDebugView_SetView1);
+		AIDebugViewInputComponent->BindKey(FInputChord(EKeys::Two, false, false, true, false), IE_Pressed, this, &UGameplayDebuggingControllerComponent::ToggleAIDebugView_SetView2);
+		AIDebugViewInputComponent->BindKey(FInputChord(EKeys::Three, false, false, true, false), IE_Pressed, this, &UGameplayDebuggingControllerComponent::ToggleAIDebugView_SetView3);
+		AIDebugViewInputComponent->BindKey(FInputChord(EKeys::Four, false, false, true, false), IE_Pressed, this, &UGameplayDebuggingControllerComponent::ToggleAIDebugView_SetView4);
+		AIDebugViewInputComponent->BindKey(FInputChord(EKeys::Five, false, false, true, false), IE_Pressed, this, &UGameplayDebuggingControllerComponent::ToggleAIDebugView_SetView5);
+		AIDebugViewInputComponent->BindKey(FInputChord(EKeys::Six, false, false, true, false), IE_Pressed, this, &UGameplayDebuggingControllerComponent::ToggleAIDebugView_SetView6);
+		AIDebugViewInputComponent->BindKey(FInputChord(EKeys::Seven, false, false, true, false), IE_Pressed, this, &UGameplayDebuggingControllerComponent::ToggleAIDebugView_SetView7);
+		AIDebugViewInputComponent->BindKey(FInputChord(EKeys::Eight, false, false, true, false), IE_Pressed, this, &UGameplayDebuggingControllerComponent::ToggleAIDebugView_SetView8);
+		AIDebugViewInputComponent->BindKey(FInputChord(EKeys::Nine, false, false, true, false), IE_Pressed, this, &UGameplayDebuggingControllerComponent::ToggleAIDebugView_SetView9);
+		AIDebugViewInputComponent->BindKey(FInputChord(EKeys::Equals, false, false, true, false), IE_Released, this, &UGameplayDebuggingControllerComponent::NextEQSQuery);
+#endif
 	}
 	
 	if (PlayerOwner.IsValid())
diff --git a/Engine/Source/Developer/GameplayDebugger/Private/GameplayDebuggingHUDComponent.cpp b/Engine/Source/Developer/GameplayDebugger/Private/GameplayDebuggingHUDComponent.cpp
--- a/Engine/Source/Developer/GameplayDebugger/Private/GameplayDebuggingHUDComponent.cpp
+++ b/Engine/Source/Developer/GameplayDebugger/Private/GameplayDebuggingHUDComponent.cpp
@@ -25,11 +25,11 @@ AGameplayDebuggingHUDComponent::AGameplayDebuggingHUDComponent(const class FPost
 	, EngineShowFlags(EShowFlagInitMode::ESFIM_Game)
 {
 	World = NULL;
-#if WITH_EDITOR
+	SetTickableWhenPaused(true);
+
+#if WITH_EDITORONLY_DATA
 	SetIsTemporarilyHiddenInEditor(true);
 	SetActorHiddenInGame(false);
-#endif
-#if WITH_EDITORONLY_DATA
 	bHiddenEdLevel = true;
 	bHiddenEdLayer = true;
 	bHiddenEd = true;
@@ -106,7 +106,7 @@ void AGameplayDebuggingHUDComponent::PrintAllData()
 		DrawDebugComponentData(MyPC, DebugComponent);
 	}
 
-	if (!EngineShowFlags.DebugAI)
+	if (DefaultContext.Canvas && DefaultContext.Canvas->SceneView && DefaultContext.Canvas->SceneView->Family && DefaultContext.Canvas->SceneView->Family->EngineShowFlags.Game)
 	{
 		DrawMenu(MenuX, MenuY, DebugComponent);
 	}
@@ -182,7 +182,7 @@ void AGameplayDebuggingHUDComponent::DrawDebugComponentData(APlayerController* M
 	OverHeadContext = FPrintContext(GEngine->GetSmallFont(), Canvas, ScreenLoc.X, ScreenLoc.Y);
 
 	FGameplayDebuggerSettings DebuggerSettings = GameplayDebuggerSettings(GetDebuggingReplicator());
-	if (DebuggerSettings.CheckFlag(EAIDebugDrawDataView::OverHead) || EngineShowFlags.DebugAI)
+	if (DebuggerSettings.CheckFlag(EAIDebugDrawDataView::OverHead) /*|| EngineShowFlags.DebugAI*/)
 	{
 		DrawOverHeadInformation(MyPC, DebugComponent);
 	}
@@ -192,9 +192,10 @@ void AGameplayDebuggingHUDComponent::DrawDebugComponentData(APlayerController* M
 		DrawNavMeshSnapshot(MyPC, DebugComponent);
 	}
 
-	if (DebugComponent->GetSelectedActor() && DebugComponent->IsSelected())
+	const bool bDrawFullData = GetDebuggingReplicator()->GetSelectedActorToDebug() == SelectedActor;
+	if (DebugComponent->GetSelectedActor() && bDrawFullData)
 	{
-		if (DebuggerSettings.CheckFlag(EAIDebugDrawDataView::Basic) || EngineShowFlags.DebugAI)
+		if (DebuggerSettings.CheckFlag(EAIDebugDrawDataView::Basic) /*|| EngineShowFlags.DebugAI*/)
 		{
 			DrawBasicData(MyPC, DebugComponent);
 		}
@@ -217,7 +218,7 @@ void AGameplayDebuggingHUDComponent::DrawDebugComponentData(APlayerController* M
 			}
 		}
 
-		if (DebuggerSettings.CheckFlag(EAIDebugDrawDataView::Perception) || EngineShowFlags.DebugAI)
+		if (DebuggerSettings.CheckFlag(EAIDebugDrawDataView::Perception) /*|| EngineShowFlags.DebugAI*/)
 		{
 			DrawPerception(MyPC, DebugComponent);
 		}
@@ -275,7 +276,7 @@ void AGameplayDebuggingHUDComponent::DrawOverHeadInformation(APlayerController*
 	FString ObjectName = FString::Printf( TEXT("{yellow}%s {white}(%s)"), *DebugComponent->ControllerName, *DebugComponent->PawnName);
 	CalulateStringSize(OverHeadContext, OverHeadContext.Font, ObjectName, TextXL, YL);
 
-	bool bDrawFullOverHead = DebugComponent->bIsSelectedForDebugging;
+	bool bDrawFullOverHead = GetDebuggingReplicator()->GetSelectedActorToDebug() == MyPawn;
 	float IconXLocation = OverHeadContext.DefaultX;
 	float IconYLocation = OverHeadContext.DefaultY;
 	if (bDrawFullOverHead)
@@ -307,8 +308,6 @@ void AGameplayDebuggingHUDComponent::DrawOverHeadInformation(APlayerController*
 		OverHeadContext.FontRenderInfo.bEnableShadow = false;
 	}
 
-	APlayerController* const MyPC = Cast<APlayerController>(PlayerOwner);
-	UGameplayDebuggingControllerComponent*  GDC = MyPC ? MyPC->FindComponentByClass<UGameplayDebuggingControllerComponent>() : NULL;
 	if (EngineShowFlags.DebugAI)
 	{
 		PrintString(OverHeadContext, FString::Printf(TEXT("{red}%s\n"), *DebugComponent->PathErrorString));
@@ -333,10 +332,7 @@ void AGameplayDebuggingHUDComponent::DrawBasicData(APlayerController* PC, class
 		PrintString(DefaultContext, FString::Printf(TEXT("{red}%s\n"), *DebugComponent->PathErrorString));
 	}
 
-	if (DebugComponent->IsSelected())
-	{
-		DrawPath(PC, DebugComponent);
-	}
+	DrawPath(PC, DebugComponent);
 #endif //!(UE_BUILD_SHIPPING || UE_BUILD_TEST)
 }
 
@@ -387,24 +383,26 @@ void AGameplayDebuggingHUDComponent::DrawEQSData(APlayerController* PC, class UG
 	{
 		APlayerController* const MyPC = Cast<APlayerController>(PlayerOwner);
 		FVector CamLocation;
-		FRotator CamRotation;
-		check(MyPC->PlayerCameraManager != NULL);
-		MyPC->PlayerCameraManager->GetCameraViewPoint(CamLocation, CamRotation);
-		FVector FireDir = CamRotation.Vector();
-
-		FVector out_Location;
-		FRotator out_Rotation;
-		MyPC->GetPlayerViewPoint(out_Location, out_Rotation);
+		FVector FireDir;
+		if (!MyPC->GetSpectatorPawn())
+		{
+			FRotator CamRotation;
+			MyPC->GetPlayerViewPoint(CamLocation, CamRotation);
+			FireDir = CamRotation.Vector();
+		}
+		else
+		{
+			FireDir = DefaultContext.Canvas->SceneView->GetViewDirection();
+			CamLocation = DefaultContext.Canvas->SceneView->ViewMatrices.ViewOrigin;
+		}
 
 		float bestAim = 0;
 		for (int32 Index = 0; Index < CurrentLocalData.RenderDebugHelpers.Num(); ++Index)
 		{
 			auto& CurrentItem = CurrentLocalData.RenderDebugHelpers[Index];
 
-			// look for best controlled pawn target
-			const FVector AimDir = CurrentItem.Location - out_Location;
+			const FVector AimDir = CurrentItem.Location - CamLocation;
 			float FireDist = AimDir.SizeSquared();
-			// only find targets which are < 25000 units away
 
 			FireDist = FMath::Sqrt(FireDist);
 			float newAim = FireDir | AimDir;
@@ -455,7 +453,7 @@ void AGameplayDebuggingHUDComponent::DrawEQSData(APlayerController* PC, class UG
 
 	const float RowHeight = 20.0f;
 	const int32 NumTests = CurrentLocalData.Tests.Num();
-	if (CurrentLocalData.NumValidItems > 0)
+	if (CurrentLocalData.NumValidItems > 0 && GetDebuggingReplicator()->EnableEQSOnHUD )
 	{
 		// draw test weights for best X items
 		const int32 NumItems = CurrentLocalData.Items.Num();
diff --git a/Engine/Source/Developer/GameplayDebugger/Private/GameplayDebuggingReplicator.cpp b/Engine/Source/Developer/GameplayDebugger/Private/GameplayDebuggingReplicator.cpp
--- a/Engine/Source/Developer/GameplayDebugger/Private/GameplayDebuggingReplicator.cpp
+++ b/Engine/Source/Developer/GameplayDebugger/Private/GameplayDebuggingReplicator.cpp
@@ -29,11 +29,11 @@ AGameplayDebuggingReplicator::AGameplayDebuggingReplicator(const class FPostCons
 	RootComponent = SceneComponent;
 
 #if WITH_EDITOR
-	SetTickableWhenPaused(true);
 	SetIsTemporarilyHiddenInEditor(true);
-	SetActorHiddenInGame(false);
 #endif
 #if WITH_EDITORONLY_DATA
+	SetTickableWhenPaused(true);
+	SetActorHiddenInGame(false);
 	bHiddenEdLevel = true;
 	bHiddenEdLayer = true;
 	bHiddenEd = true;
@@ -42,6 +42,22 @@ AGameplayDebuggingReplicator::AGameplayDebuggingReplicator(const class FPostCons
 
 	DebuggerShowFlags =  GameplayDebuggerSettings().DebuggerShowFlags;
 
+	FGameplayDebuggerSettings Settings = GameplayDebuggerSettings(this);
+#define UPDATE_VIEW_PROPS(__FlagName__)  __FlagName__ = Settings.CheckFlag(EAIDebugDrawDataView::__FlagName__);
+	UPDATE_VIEW_PROPS(OverHead);
+	UPDATE_VIEW_PROPS(Basic);
+	UPDATE_VIEW_PROPS(BehaviorTree);
+	UPDATE_VIEW_PROPS(EQS);
+	UPDATE_VIEW_PROPS(Perception);
+	UPDATE_VIEW_PROPS(GameView1);
+	UPDATE_VIEW_PROPS(GameView2);
+	UPDATE_VIEW_PROPS(GameView3);
+	UPDATE_VIEW_PROPS(GameView4);
+	UPDATE_VIEW_PROPS(GameView5);
+#undef UPDATE_VIEW_PROPS
+
+	EnableEQSOnHUD = true;
+
 	if (!HasAnyFlags(RF_ClassDefaultObject))
 	{
 		SetActorTickEnabled(true);
@@ -62,6 +78,7 @@ void AGameplayDebuggingReplicator::GetLifetimeReplicatedProps(TArray< FLifetimeP
 		DOREPLIFETIME_CONDITION(AGameplayDebuggingReplicator, DebugComponent, COND_OwnerOnly);
 		DOREPLIFETIME_CONDITION(AGameplayDebuggingReplicator, LocalPlayerOwner, COND_OwnerOnly);
 		DOREPLIFETIME_CONDITION(AGameplayDebuggingReplicator, bIsGlobalInWorld, COND_OwnerOnly);
+		DOREPLIFETIME_CONDITION(AGameplayDebuggingReplicator, LastSelectedActorToDebug, COND_OwnerOnly);
 #endif
 }
 
@@ -76,13 +93,77 @@ void AGameplayDebuggingReplicator::PostInitializeComponents()
 	SetActorTickEnabled(true);
 }
 
+#if WITH_EDITOR
+void AGameplayDebuggingReplicator::PostEditChangeProperty(FPropertyChangedEvent& PropertyChangedEvent)
+{
+	Super::PostEditChangeProperty(PropertyChangedEvent);
+
+	if (!PropertyChangedEvent.Property)
+	{
+		return;
+	}
+
+	FGameplayDebuggerSettings Settings = GameplayDebuggerSettings(this);
+
+#define CHECK_AND_UPDATE_FLAGS(__FlagName_) \
+	if (PropertyChangedEvent.Property->GetFName() == GET_MEMBER_NAME_CHECKED(AGameplayDebuggingReplicator, __FlagName_)) \
+	{ \
+		__FlagName_ ? Settings.SetFlag(EAIDebugDrawDataView::__FlagName_) : Settings.ClearFlag(EAIDebugDrawDataView::__FlagName_); \
+		GetDebugComponent()->ServerReplicateData(Settings.CheckFlag(EAIDebugDrawDataView::__FlagName_) ? EDebugComponentMessage::ActivateDataView : EDebugComponentMessage::DeactivateDataView, EAIDebugDrawDataView::__FlagName_); \
+	}else
+
+	CHECK_AND_UPDATE_FLAGS(OverHead)
+	CHECK_AND_UPDATE_FLAGS(Basic)
+	CHECK_AND_UPDATE_FLAGS(BehaviorTree)
+	CHECK_AND_UPDATE_FLAGS(EQS)
+	CHECK_AND_UPDATE_FLAGS(Perception)
+	CHECK_AND_UPDATE_FLAGS(GameView1)
+	CHECK_AND_UPDATE_FLAGS(GameView2)
+	CHECK_AND_UPDATE_FLAGS(GameView3)
+	CHECK_AND_UPDATE_FLAGS(GameView4)
+	CHECK_AND_UPDATE_FLAGS(GameView5) {}
+
+#if WITH_EQS
+	if (PropertyChangedEvent.Property->GetFName() == GET_MEMBER_NAME_CHECKED(AGameplayDebuggingReplicator, EQS))
+	{
+		GetDebugComponent()->EnableClientEQSSceneProxy(Settings.CheckFlag(EAIDebugDrawDataView::EQS));
+		GetDebugComponent()->MarkRenderStateDirty();
+	}
+
+	if (PropertyChangedEvent.Property->GetFName() == GET_MEMBER_NAME_CHECKED(AGameplayDebuggingReplicator, ActiveEQSIndex))
+	{
+		GetDebugComponent()->SetEQSIndex(ActiveEQSIndex);
+	}
+#endif // WITH_EQS
+
+#undef CHECK_AND_UPDATE_FLAGS
+}
+#endif
+
 void AGameplayDebuggingReplicator::BeginPlay()
 {
 	Super::BeginPlay();
 
+	if (Role == ROLE_Authority)
+	{
+		bReplicates = false;
+		SetRemoteRoleForBackwardsCompat(ROLE_SimulatedProxy);
+		SetReplicates(true);
+
+		if (!DebugComponentClass.IsValid() && GetWorld() && GetNetMode() < ENetMode::NM_Client)
+		{
+			DebugComponentClass = StaticLoadClass(UGameplayDebuggingComponent::StaticClass(), NULL, *DebugComponentClassName, NULL, LOAD_None, NULL);
+			if (!DebugComponentClass.IsValid())
+			{
+				DebugComponentClass = UGameplayDebuggingComponent::StaticClass();
+			}
+		}
+		GetDebugComponent();
+	}
+
 #if WITH_EDITOR
 	const UEditorEngine* EEngine = Cast<UEditorEngine>(GEngine);
-	if (EEngine && (EEngine->bIsSimulatingInEditor || EEngine->EditorWorld) && IsGlobalInWorld() && GCurrentLevelEditingViewportClient->EngineShowFlags.DebugAI)
+	if (EEngine && (EEngine->bIsSimulatingInEditor || EEngine->EditorWorld) && GetWorld() != EEngine->EditorWorld && !IsGlobalInWorld() && GCurrentLevelEditingViewportClient->EngineShowFlags.DebugAI)
 	{
 		SetIsTemporarilyHiddenInEditor(false);
 		SetActorHiddenInGame(false);
@@ -90,6 +171,11 @@ void AGameplayDebuggingReplicator::BeginPlay()
 		bHiddenEdLayer = false;
 		bHiddenEd = false;
 		bEditable = true;
+
+		if (DebugComponent)
+		{
+			DebugComponent->ServerReplicateData(EDebugComponentMessage::ActivateReplication, EAIDebugDrawDataView::Empty);
+		}
 	}
 	else
 	{
@@ -100,32 +186,12 @@ void AGameplayDebuggingReplicator::BeginPlay()
 		bHiddenEdLayer = true;
 		bHiddenEd = true;
 		bEditable = false;
-	}
-#endif
-
-	if (Role == ROLE_Authority)
-	{
-		bReplicates = false;
-		SetRemoteRoleForBackwardsCompat(ROLE_SimulatedProxy);
-		SetReplicates(true);
-
-		if (!DebugComponentClass.IsValid())
+		if (DebugComponent)
 		{
-			DebugComponentClass = StaticLoadClass(UGameplayDebuggingComponent::StaticClass(), NULL, *DebugComponentClassName, NULL, LOAD_None, NULL);
-			if (!DebugComponentClass.IsValid())
-			{
-				DebugComponentClass = UGameplayDebuggingComponent::StaticClass();
-			}
-		}
-
-		if (DebugComponent == NULL)
-		{
-			DebugComponent = ConstructObject<UGameplayDebuggingComponent>(DebugComponentClass.Get(), this);
-			DebugComponent->SetIsReplicated(true);
-			DebugComponent->RegisterComponent();
-			//DebugComponent->Activate();
+			DebugComponent->ServerReplicateData(EDebugComponentMessage::DeactivateDataView, EAIDebugDrawDataView::Empty);
 		}
 	}
+#endif
 
 	if (GetWorld() && GetNetMode() != ENetMode::NM_DedicatedServer)
 	{
@@ -145,18 +211,6 @@ void AGameplayDebuggingReplicator::BeginPlay()
 		}
 	}
 
-
-	if (!DebugComponentClass.IsValid() && GetWorld() && GetNetMode() < ENetMode::NM_Client)
-	{
-		DebugComponentClass = StaticLoadClass(UGameplayDebuggingComponent::StaticClass(), NULL, *DebugComponentClassName, NULL, LOAD_None, NULL);
-		if (!DebugComponentClass.IsValid())
-		{
-			DebugComponentClass = UGameplayDebuggingComponent::StaticClass();
-		}
-	}
-
-	SetActorTickEnabled(true);
-	SetTickableWhenPaused(true);
 }
 
 UGameplayDebuggingComponent* AGameplayDebuggingReplicator::GetDebugComponent()
@@ -213,6 +267,12 @@ bool AGameplayDebuggingReplicator::ServerReplicateMessage_Validate(class AActor*
 void AGameplayDebuggingReplicator::ServerReplicateMessage_Implementation(class  AActor* Actor, uint32 InMessage, uint32 DataView)
 {
 #if !(UE_BUILD_SHIPPING || UE_BUILD_TEST)
+	if ((EDebugComponentMessage::Type)InMessage == EDebugComponentMessage::DeactivateReplilcation)
+	{
+		SetActorToDebug(NULL);
+		MarkComponentsRenderStateDirty();
+	}
+
 	if (GetDebugComponent())
 	{
 		GetDebugComponent()->ServerReplicateData((EDebugComponentMessage::Type)InMessage, (EAIDebugDrawDataView::Type)DataView);
@@ -234,7 +294,13 @@ void AGameplayDebuggingReplicator::EnableDraw(bool bEnable)
 		GameHUD->bShowHUD = bEnable ? false : true;
 	}
 	GEngine->bEnableOnScreenDebugMessages = bEnable ? false : true;
-	GetDebugComponent()->MarkRenderStateDirty();
+
+	if (DebugComponent)
+	{
+		const bool bEnabledEQSView = GameplayDebuggerSettings(this).CheckFlag(EAIDebugDrawDataView::EQS);
+		DebugComponent->EnableClientEQSSceneProxy(bEnable && bEnabledEQSView ? true : false);
+		DebugComponent->MarkRenderStateDirty();
+	}
 }
 
 bool AGameplayDebuggingReplicator::IsToolCreated()
@@ -306,9 +372,9 @@ void AGameplayDebuggingReplicator::TickActor(float DeltaTime, enum ELevelTick Ti
 
 void AGameplayDebuggingReplicator::SetActorToDebug(AActor* InActor) 
 { 
-	if (LastSelectedActorInSimulation.Get() != InActor)
+	if (LastSelectedActorToDebug != InActor)
 	{
-		LastSelectedActorInSimulation = InActor;
+		LastSelectedActorToDebug = InActor;
 		UGameplayDebuggingComponent::OnDebuggingTargetChangedDelegate.Broadcast(InActor, InActor ? InActor->IsSelected() : false);
 		APawn* TargetPawn = Cast<APawn>(InActor);
 		if (TargetPawn)
@@ -320,7 +386,6 @@ void AGameplayDebuggingReplicator::SetActorToDebug(AActor* InActor)
 	if (UGameplayDebuggingComponent* DebugComponent = GetDebugComponent())
 	{
 		DebugComponent->SetActorToDebug(InActor);
-		DebugComponent->SelectForDebugging(InActor ? InActor->IsSelected() : false);
 	}
 }
 
@@ -353,7 +418,6 @@ void AGameplayDebuggingReplicator::OnDebugAIDelegate(class UCanvas* Canvas, clas
 	UWorld* World = GetWorld();
 	if (World && GetDebugComponent() && GetDebugComponent()->GetOwnerRole() == ROLE_Authority)
 	{
-		bool bBroadcastedNewSelection = false;
 		UGameplayDebuggingControllerComponent*  GDC = FindComponentByClass<UGameplayDebuggingControllerComponent>();
 		TArray<int32> OryginalReplicateViewDataCounters;
 
@@ -367,26 +431,35 @@ void AGameplayDebuggingReplicator::OnDebugAIDelegate(class UCanvas* Canvas, clas
 		}
 
 		// looks like Simulate in UE4 Editor - let's find selected Pawn to debug
+		AActor* FullSelectedTarget = NULL;
 		for (FConstPawnIterator Iterator = World->GetPawnIterator(); Iterator; ++Iterator)
 		{
 			AActor* NewTarget = Cast<AActor>(*Iterator);
 
+			if (NewTarget->IsSelected() && !FullSelectedTarget)
+			{
+				FullSelectedTarget = NewTarget;
+				continue;
+			}
+
 			//We needs to collect data manually in Simulate
 			GetDebugComponent()->SetActorToDebug(NewTarget);
-			GetDebugComponent()->SelectForDebugging(NewTarget->IsSelected());
 			GetDebugComponent()->CollectDataToReplicate(NewTarget->IsSelected());
+			DrawDebugData(Canvas, PC);
+		}
 
-			if (NewTarget->IsSelected())
-			{
-				SetActorToDebug(NewTarget);
-				bBroadcastedNewSelection = true;
-			}
+		const AActor* OldActor = LastSelectedActorToDebug;
+		SetActorToDebug(FullSelectedTarget);
+		if (FullSelectedTarget)
+		{
+			GetDebugComponent()->CollectDataToReplicate(true);
+			GetDebugComponent()->SetEQSIndex(ActiveEQSIndex);
 			DrawDebugData(Canvas, PC);
 		}
 
-		if (!bBroadcastedNewSelection && GetSelectedActorToDebug() && !GetSelectedActorToDebug()->IsSelected())
+		if (GetSelectedActorToDebug() != OldActor)
 		{
-			SetActorToDebug(NULL);
+			GetDebugComponent()->MarkRenderStateDirty();
 		}
 
 		if (UGameplayDebuggingComponent* DebugComponent = GetDebugComponent())
