Fixes for GameplayDebugger.
- fixed activation in different configurations (PIE, standalone, client-server, etc.).
- Fixed issues related to NavMesh rendering
- fixed EQS query switching
Fixes integrated from Dev (CL #2292104 and CL #2292198)

diff --git a/Engine/Config/BaseInput.ini b/Engine/Config/BaseInput.ini
--- a/Engine/Config/BaseInput.ini
+++ b/Engine/Config/BaseInput.ini
@@ -42,7 +42,7 @@ bRequireCtrlToNavigateAutoComplete=False
 +DebugExecBindings=(Key=Enter,Command="TOGGLE_FULLSCREEN",Alt=True)
 +DebugExecBindings=(Key=Tab,Command="FocusNextPIEWindow",Control=True)
 +DebugExecBindings=(Key=Tab,Command="FocusLastPIEWindow",Control=True,Shift=True)
-+DebugExecBindings=(Key=Quote,Command="cheat EnableGDT",Control=False,Shift=False)
++DebugExecBindings=(Key=Quote,Command="EnableGDT",Control=False,Shift=False)
 
 [/Script/Engine.Console]
 HistoryBot=-1
diff --git a/Engine/Source/Developer/GameplayDebugger/Classes/GameplayDebuggingComponent.h b/Engine/Source/Developer/GameplayDebugger/Classes/GameplayDebuggingComponent.h
--- a/Engine/Source/Developer/GameplayDebugger/Classes/GameplayDebuggingComponent.h
+++ b/Engine/Source/Developer/GameplayDebugger/Classes/GameplayDebuggingComponent.h
@@ -150,7 +150,7 @@ class GAMEPLAYDEBUGGER_API UGameplayDebuggingComponent : public UPrimitiveCompon
 		return TargetActor;
 	}
 
-	void SetEQSIndex(int32 Index) { CurrentEQSIndex = AllEQSName.Num() > 0 ? FMath::Clamp(Index, 0, AllEQSName.Num() - 1) : INDEX_NONE; }
+	void SetEQSIndex(int32 Index) { CurrentEQSIndex = Index; }
 	int32 GetEQSIndex() { return CurrentEQSIndex; }
 	//=============================================================================
 	// EQS debugging
diff --git a/Engine/Source/Developer/GameplayDebugger/Private/GameplayDebugger.cpp b/Engine/Source/Developer/GameplayDebugger/Private/GameplayDebugger.cpp
--- a/Engine/Source/Developer/GameplayDebugger/Private/GameplayDebugger.cpp
+++ b/Engine/Source/Developer/GameplayDebugger/Private/GameplayDebugger.cpp
@@ -1,9 +1,12 @@
 // Copyright 1998-2014 Epic Games, Inc. All Rights Reserved.
 
 #include "GameplayDebuggerPrivate.h"
+#include "Misc/CoreMisc.h"
 #include "GameplayDebuggingComponent.h"
 #include "GameplayDebuggingHUDComponent.h"
 #include "GameplayDebuggingReplicator.h"
+#include "GameplayDebuggingControllerComponent.h"
+#include "AISystem.h"
 
 uint32 FGameplayDebuggerSettings::ShowFlagIndex = 0;
 
@@ -13,7 +16,7 @@ FGameplayDebuggerSettings GameplayDebuggerSettings(class AGameplayDebuggingRepli
 	return FGameplayDebuggerSettings(Replicator == NULL ? Settings : Replicator->DebuggerShowFlags);
 }
 
-class FGameplayDebugger : public IGameplayDebugger
+class FGameplayDebugger : public FSelfRegisteringExec, public IGameplayDebugger
 {
 public:
 	// Begin IModuleInterface
@@ -32,6 +35,10 @@ public:
 	void AddReplicator(UWorld* InWorld, AGameplayDebuggingReplicator* InReplicator);
 	void RemoveReplicator(UWorld* InWorld, AGameplayDebuggingReplicator* InReplicator);
 
+	// Begin FExec Interface
+	virtual bool Exec(UWorld* Inworld, const TCHAR* Cmd, FOutputDevice& Ar) override;
+	// End FExec Interface
+
 private:
 	virtual bool CreateGameplayDebuggerForPlayerController(APlayerController* PlayerController) override;
 
@@ -251,4 +258,115 @@ void FGameplayDebugger::OnLevelActorDeleted(AActor* InActor)
 	}
 #endif
 }
-#endif
+#endif //WITH_EDITOR
+
+bool FGameplayDebugger::Exec(UWorld* Inworld, const TCHAR* Cmd, FOutputDevice& Ar)
+{
+	bool bHandled = false;
+
+	APlayerController* PC = Inworld ? Inworld->GetFirstPlayerController() : NULL;
+	if (!Inworld || !PC)
+	{
+		return bHandled;
+	}
+
+	if (FParse::Command(&Cmd, TEXT("EnableGDT")))
+	{
+		AGameplayDebuggingReplicator* Replicator = NULL;
+		for (TActorIterator<AGameplayDebuggingReplicator> It(Inworld); It; ++It)
+		{
+			Replicator = *It;
+			if (Replicator && !Replicator->IsPendingKill())
+			{
+				APlayerController* LocalPC = Replicator->GetLocalPlayerOwner();
+				if (LocalPC == PC)
+				{
+					break;
+				}
+			}
+			Replicator = NULL;
+		}
+
+		if (Inworld->GetNetMode() == NM_Client)
+		{
+			if (!Replicator)
+			{
+				PC->ConsoleCommand("cheat EnableGDT");
+			}
+			else if (!Replicator->IsToolCreated())
+			{
+				Replicator->CreateTool();
+				Replicator->EnableTool();
+				bHandled = true;
+			}
+		}
+		else if (Inworld->GetNetMode() < NM_Client)
+		{
+			if (!Replicator)
+			{
+				CreateGameplayDebuggerForPlayerController(PC);
+				for (TActorIterator<AGameplayDebuggingReplicator> It(Inworld); It; ++It)
+				{
+					Replicator = *It;
+					if (Replicator && !Replicator->IsPendingKill())
+					{
+						APlayerController* LocalPC = Replicator->GetLocalPlayerOwner();
+						if (LocalPC == PC)
+						{
+							break;
+						}
+					}
+					Replicator = NULL;
+				}
+			}
+
+			if (Inworld->GetNetMode() != NM_DedicatedServer)
+			{
+				if (Replicator && !Replicator->IsToolCreated())
+				{
+					Replicator->CreateTool();
+					Replicator->EnableTool();
+					bHandled = true;
+				}
+			}
+		}
+	}
+	else if (FParse::Command(&Cmd, TEXT("RunEQS")))
+	{
+		bHandled = true;
+		APlayerController* MyPC = Inworld->GetFirstPlayerController();
+		UAISystem* AISys = UAISystem::GetCurrent(Inworld);
+
+		UEnvQueryManager* EQS = AISys ? AISys->GetEnvironmentQueryManager() : NULL;
+		if (MyPC && EQS)
+		{
+			AGameplayDebuggingReplicator* DebuggingReplicator = NULL;
+			for (FActorIterator It(Inworld); It; ++It)
+			{
+				AActor* A = *It;
+				if (A && A->IsA(AGameplayDebuggingReplicator::StaticClass()) && !A->IsPendingKill())
+				{
+					DebuggingReplicator = Cast<AGameplayDebuggingReplicator>(A);
+					if (DebuggingReplicator && !DebuggingReplicator->IsGlobalInWorld() && DebuggingReplicator->GetLocalPlayerOwner() == MyPC)
+					{
+						break;
+					}
+				}
+			}
+
+			UObject* Target = DebuggingReplicator != NULL ? DebuggingReplicator->GetSelectedActorToDebug() : NULL;
+			FString QueryName = FParse::Token(Cmd, 0);
+			if (Target)
+			{
+				AISys->RunEQS(QueryName, Target);
+			}
+			else
+			{
+				MyPC->ClientMessage(TEXT("No debugging target to run EQS"));
+			}
+		}
+	}
+
+	return bHandled;
+}
+
diff --git a/Engine/Source/Developer/GameplayDebugger/Private/GameplayDebuggerExec.cpp b/Engine/Source/Developer/GameplayDebugger/Private/GameplayDebuggerExec.cpp
deleted file mode 100644
diff --git a/Engine/Source/Developer/GameplayDebugger/Private/GameplayDebuggingComponent.cpp b/Engine/Source/Developer/GameplayDebugger/Private/GameplayDebuggingComponent.cpp
--- a/Engine/Source/Developer/GameplayDebugger/Private/GameplayDebuggingComponent.cpp
+++ b/Engine/Source/Developer/GameplayDebugger/Private/GameplayDebuggingComponent.cpp
@@ -227,7 +227,6 @@ void UGameplayDebuggingComponent::SetActorToDebug(AActor* Actor)
 	TargetActor = Actor;
 	EQSLocalData.Reset();
 	AllEQSName.Reset();
-	CurrentEQSIndex = INDEX_NONE;
 #endif //!(UE_BUILD_SHIPPING || UE_BUILD_TEST)
 }
 
@@ -253,7 +252,6 @@ void UGameplayDebuggingComponent::TickComponent(float DeltaTime, enum ELevelTick
 	}
 
 	AGameplayDebuggingReplicator* Replicator = Cast<AGameplayDebuggingReplicator>(GetOwner());
-	CurrentEQSIndex = AllEQSName.Num() > 0 ?FMath::Clamp(CurrentEQSIndex, 0, AllEQSName.Num() - 1) : INDEX_NONE;
 
 #endif //!(UE_BUILD_SHIPPING || UE_BUILD_TEST)
 }
@@ -600,7 +598,6 @@ void UGameplayDebuggingComponent::OnRep_UpdateEQS()
 		ArReader << EQSLocalData;
 	}	
 
-	CurrentEQSIndex = AllEQSName.Num() > 0 ? FMath::Clamp(CurrentEQSIndex, 0, AllEQSName.Num() - 1) : INDEX_NONE;
 	UpdateBounds();
 	MarkRenderStateDirty();
 #endif //USE_EQS_DEBUGGER
@@ -1084,14 +1081,20 @@ void UGameplayDebuggingComponent::PrepareNavMeshData(struct FNavMeshSceneProxyDa
 //----------------------------------------------------------------------//
 FPrimitiveSceneProxy* UGameplayDebuggingComponent::CreateSceneProxy()
 {
-	if (!IsActive())
+	FDebugRenderSceneCompositeProxy* CompositeProxy = NULL;
+	AGameplayDebuggingReplicator* Replicator = Cast<AGameplayDebuggingReplicator>(GetOwner());
+	if (!World || World->GetNetMode() == NM_DedicatedServer)
+	{
+		return NULL;
+	}
+
+	if (!Replicator || !Replicator->IsDrawEnabled())
 	{
 		return NULL;
 	}
-	FDebugRenderSceneCompositeProxy* CompositeProxy = NULL;
 
 #if WITH_RECAST	
-	if (ShouldReplicateData(EAIDebugDrawDataView::NavMesh) && World && World->GetNetMode() != NM_DedicatedServer)
+	if (ShouldReplicateData(EAIDebugDrawDataView::NavMesh))
 	{
 		FNavMeshSceneProxyData NewNavmeshRenderData;
 		NewNavmeshRenderData.Reset();
@@ -1108,11 +1111,11 @@ FPrimitiveSceneProxy* UGameplayDebuggingComponent::CreateSceneProxy()
 #if USE_EQS_DEBUGGER
 	if (ShouldReplicateData(EAIDebugDrawDataView::EQS) && IsClientEQSSceneProxyEnabled())
 	{
-		if (EQSLocalData.IsValidIndex(CurrentEQSIndex))
+		const int32 EQSIndex = AllEQSName.Num() > 0 ? FMath::Clamp(CurrentEQSIndex, 0, AllEQSName.Num() - 1) : INDEX_NONE;
+		if (EQSLocalData.IsValidIndex(EQSIndex))
 		{
 			CompositeProxy = CompositeProxy ? CompositeProxy : (new FDebugRenderSceneCompositeProxy(this));
-			auto& CurrentLocalData = EQSLocalData[CurrentEQSIndex];
-			bool bUseDebugAIFlag = false;
+			auto& CurrentLocalData = EQSLocalData[EQSIndex];
 			CompositeProxy->AddChild(new FEQSSceneProxy(this, TEXT("DebugAI"), false, CurrentLocalData.SolidSpheres, CurrentLocalData.Texts));
 		}
 	}
diff --git a/Engine/Source/Developer/GameplayDebugger/Private/GameplayDebuggingControllerComponent.cpp b/Engine/Source/Developer/GameplayDebugger/Private/GameplayDebuggingControllerComponent.cpp
--- a/Engine/Source/Developer/GameplayDebugger/Private/GameplayDebuggingControllerComponent.cpp
+++ b/Engine/Source/Developer/GameplayDebugger/Private/GameplayDebuggingControllerComponent.cpp
@@ -129,7 +129,7 @@ void UGameplayDebuggingControllerComponent::BindActivationKeys()
 		FInputChord ActivationKey(EKeys::Quote, false, false, false, false);
 		for (uint32 BindIndex = 0; BindIndex < (uint32)PlayerOwner->PlayerInput->DebugExecBindings.Num(); BindIndex++)
 		{
-			if (PlayerOwner->PlayerInput->DebugExecBindings[BindIndex].Command == TEXT("cheat EnableGDT"))
+			if (PlayerOwner->PlayerInput->DebugExecBindings[BindIndex].Command == TEXT("EnableGDT"))
 			{
 				ActivationKey.Key = PlayerOwner->PlayerInput->DebugExecBindings[BindIndex].Key;
 				ActivationKey.bCtrl = PlayerOwner->PlayerInput->DebugExecBindings[BindIndex].Control;
@@ -284,8 +284,6 @@ void UGameplayDebuggingControllerComponent::ToggleAIDebugView_SetView0()
 				UpdateNavMeshTimer();
 
 				GetDebuggingReplicator()->ServerReplicateMessage(Pawn, EDebugComponentMessage::ActivateDataView, EAIDebugDrawDataView::NavMesh);
-				OwnerComp->SetVisibility(true, true);
-				OwnerComp->SetHiddenInGame(false, true);
 				OwnerComp->MarkRenderStateDirty();
 			}
 			else
@@ -294,8 +292,6 @@ void UGameplayDebuggingControllerComponent::ToggleAIDebugView_SetView0()
 
 				GetDebuggingReplicator()->ServerReplicateMessage(Pawn, EDebugComponentMessage::DeactivateDataView, EAIDebugDrawDataView::NavMesh);
 				OwnerComp->ServerDiscardNavmeshData();
-				OwnerComp->SetVisibility(false, true);
-				OwnerComp->SetHiddenInGame(true, true);
 				OwnerComp->MarkRenderStateDirty();
 			}
 		}
diff --git a/Engine/Source/Developer/GameplayDebugger/Private/GameplayDebuggingHUDComponent.cpp b/Engine/Source/Developer/GameplayDebugger/Private/GameplayDebuggingHUDComponent.cpp
--- a/Engine/Source/Developer/GameplayDebugger/Private/GameplayDebuggingHUDComponent.cpp
+++ b/Engine/Source/Developer/GameplayDebugger/Private/GameplayDebuggingHUDComponent.cpp
@@ -351,16 +351,17 @@ void AGameplayDebuggingHUDComponent::DrawEQSData(APlayerController* PC, class UG
 #if !(UE_BUILD_SHIPPING || UE_BUILD_TEST) && WITH_EQS
 	PrintString(DefaultContext, TEXT("\n{green}EQS {white}[Use + key to switch query]\n"));
 
-	if (DebugComponent->EQSLocalData.Num() == 0)
+	if (DebugComponent->AllEQSName.Num() == 0)
 	{
 		return;
 	}
 
-	if (!DebugComponent->EQSLocalData.IsValidIndex(DebugComponent->CurrentEQSIndex))
+	const int32 EQSIndex = DebugComponent->AllEQSName.Num() > 0 ? FMath::Clamp(DebugComponent->CurrentEQSIndex, 0, DebugComponent->AllEQSName.Num() - 1) : INDEX_NONE;
+	if (!DebugComponent->AllEQSName.IsValidIndex(EQSIndex))
 	{
 		return;
 	}
-	const FString SelecterEQSName = DebugComponent->EQSLocalData[DebugComponent->CurrentEQSIndex].Name;
+	const FString SelecterEQSName = DebugComponent->AllEQSName[EQSIndex];
 
 	PrintString(DefaultContext, TEXT("{white}Queries: "));
 	for (auto CurrentQuery : DebugComponent->EQSLocalData)
@@ -376,7 +377,7 @@ void AGameplayDebuggingHUDComponent::DrawEQSData(APlayerController* PC, class UG
 	}
 	PrintString(DefaultContext, TEXT("\n"));
 
-	auto& CurrentLocalData = DebugComponent->EQSLocalData[DebugComponent->CurrentEQSIndex];
+	auto& CurrentLocalData = DebugComponent->EQSLocalData[EQSIndex];
 
 	/** find and draw item selection */
 	int32 BestItemIndex = INDEX_NONE;
@@ -518,7 +519,8 @@ void AGameplayDebuggingHUDComponent::DrawEQSItemDetails(int32 ItemIdx, class UGa
 	const float PosY = DefaultContext.CursorY + 1.0f;
 	float PosX = DefaultContext.CursorX;
 
-	auto& CurrentLocalData = DebugComponent->EQSLocalData[DebugComponent->CurrentEQSIndex];
+	const int32 EQSIndex = DebugComponent->AllEQSName.Num() > 0 ? FMath::Clamp(DebugComponent->CurrentEQSIndex, 0, DebugComponent->AllEQSName.Num() - 1) : INDEX_NONE;
+	auto& CurrentLocalData = DebugComponent->EQSLocalData[EQSIndex];
 	const EQSDebug::FItemData& ItemData = CurrentLocalData.Items[ItemIdx];
 
 	PrintString(DefaultContext, FColor::White, ItemData.Desc, PosX, PosY);
diff --git a/Engine/Source/Developer/GameplayDebugger/Private/GameplayDebuggingReplicator.cpp b/Engine/Source/Developer/GameplayDebugger/Private/GameplayDebuggingReplicator.cpp
--- a/Engine/Source/Developer/GameplayDebugger/Private/GameplayDebuggingReplicator.cpp
+++ b/Engine/Source/Developer/GameplayDebugger/Private/GameplayDebuggingReplicator.cpp
@@ -126,7 +126,8 @@ void AGameplayDebuggingReplicator::PostEditChangeProperty(FPropertyChangedEvent&
 #if WITH_EQS
 	if (PropertyChangedEvent.Property->GetFName() == GET_MEMBER_NAME_CHECKED(AGameplayDebuggingReplicator, EQS))
 	{
-		GetDebugComponent()->EnableClientEQSSceneProxy(Settings.CheckFlag(EAIDebugDrawDataView::EQS));
+		GetDebugComponent()->EnableClientEQSSceneProxy(EQS);
+		GetDebugComponent()->SetEQSIndex(ActiveEQSIndex);
 		GetDebugComponent()->MarkRenderStateDirty();
 	}
 
@@ -161,6 +162,24 @@ void AGameplayDebuggingReplicator::BeginPlay()
 		GetDebugComponent();
 	}
 
+	if (GetWorld() && GetNetMode() != ENetMode::NM_DedicatedServer)
+	{
+		if (GIsEditor)
+		{
+			UDebugDrawService::Register(TEXT("DebugAI"), FDebugDrawDelegate::CreateUObject(this, &AGameplayDebuggingReplicator::OnDebugAIDelegate));
+		}
+		UDebugDrawService::Register(TEXT("Game"), FDebugDrawDelegate::CreateUObject(this, &AGameplayDebuggingReplicator::DrawDebugDataDelegate));
+
+		if (!DebugComponentHUDClass.IsValid())
+		{
+			DebugComponentHUDClass = StaticLoadClass(AGameplayDebuggingHUDComponent::StaticClass(), NULL, *DebugComponentHUDClassName, NULL, LOAD_None, NULL);
+			if (!DebugComponentHUDClass.IsValid())
+			{
+				DebugComponentHUDClass = AGameplayDebuggingHUDComponent::StaticClass();
+			}
+		}
+	}
+
 #if WITH_EDITOR
 	const UEditorEngine* EEngine = Cast<UEditorEngine>(GEngine);
 	if (EEngine && (EEngine->bIsSimulatingInEditor || EEngine->EditorWorld) && GetWorld() != EEngine->EditorWorld && !IsGlobalInWorld() && GCurrentLevelEditingViewportClient->EngineShowFlags.DebugAI)
@@ -415,6 +434,7 @@ void AGameplayDebuggingReplicator::OnDebugAIDelegate(class UCanvas* Canvas, clas
 		return;
 	}
 
+	EnableDraw(true);
 	UWorld* World = GetWorld();
 	if (World && GetDebugComponent() && GetDebugComponent()->GetOwnerRole() == ROLE_Authority)
 	{
