Integrated CL #2299396 "Extended GameplayDebugger to collect EQS data from Pawn and AIController simultaneously (some queries can be executed on pawn and some on controller in the same time)" from Dev to Main #ue4

diff --git a/Engine/Source/Developer/GameplayDebugger/Private/GameplayDebuggingComponent.cpp b/Engine/Source/Developer/GameplayDebugger/Private/GameplayDebuggingComponent.cpp
--- a/Engine/Source/Developer/GameplayDebugger/Private/GameplayDebuggingComponent.cpp
+++ b/Engine/Source/Developer/GameplayDebugger/Private/GameplayDebuggingComponent.cpp
@@ -621,6 +621,12 @@ void UGameplayDebuggingComponent::CollectEQSData()
 	}
 
 	auto AllQueries = QueryManager->GetDebugger().GetAllQueriesForOwner(Owner);
+	const class APawn* OwnerAsPawn = Cast<class APawn>(Owner);
+	if (OwnerAsPawn != NULL && OwnerAsPawn->GetController())
+	{
+		const auto& AllControllerQueries = QueryManager->GetDebugger().GetAllQueriesForOwner(OwnerAsPawn->GetController());
+		AllQueries.Append(AllControllerQueries);
+	}
 	for (int32 Index = 0; Index < AllQueries.Num(); ++Index)
 	{
 		EQSDebug::FQueryData* CurrentLocalData = NULL;
