Fixed GameplayDebugger's EQS debugging when we have more than one query with the same name #ue4
Removed unnecessary array of EQS names from replication (we have names in replicated blobs with EQS data) #ue4
Only last 5 EQS queries are replicated (this number can be configured in *.ini file) - usually we don't need older queries to replicate and to debug. #ue4

diff --git a/Engine/Config/BaseEngine.ini b/Engine/Config/BaseEngine.ini
--- a/Engine/Config/BaseEngine.ini
+++ b/Engine/Config/BaseEngine.ini
@@ -2227,6 +2227,7 @@ StartingEnginePort=5000
 StartingBeaconPort=15000
 
 [/Script/GameplayDebugger.GameplayDebuggingReplicator]
+MaxEQSQueries=5
 DebugComponentClassName="/Script/GameplayDebugger.GameplayDebuggingComponent"
 DebugComponentHUDClassName="/Script/GameplayDebugger.GameplayDebuggingHUDComponent"
 
diff --git a/Engine/Source/Developer/GameplayDebugger/Classes/GameplayDebuggingComponent.h b/Engine/Source/Developer/GameplayDebugger/Classes/GameplayDebuggingComponent.h
--- a/Engine/Source/Developer/GameplayDebugger/Classes/GameplayDebuggingComponent.h
+++ b/Engine/Source/Developer/GameplayDebugger/Classes/GameplayDebuggingComponent.h
@@ -85,9 +85,6 @@ class GAMEPLAYDEBUGGER_API UGameplayDebuggingComponent : public UPrimitiveCompon
 	
 	/** Begin EQS replication data */
 
-	UPROPERTY(Replicated)
-	TArray<FString> AllEQSName;
-
 	UPROPERTY(ReplicatedUsing = OnRep_UpdateEQS)
 	TArray<uint8> EQSRepData;
 	
diff --git a/Engine/Source/Developer/GameplayDebugger/Classes/GameplayDebuggingReplicator.h b/Engine/Source/Developer/GameplayDebugger/Classes/GameplayDebuggingReplicator.h
--- a/Engine/Source/Developer/GameplayDebugger/Classes/GameplayDebuggingReplicator.h
+++ b/Engine/Source/Developer/GameplayDebugger/Classes/GameplayDebuggingReplicator.h
@@ -23,6 +23,9 @@ class GAMEPLAYDEBUGGER_API AGameplayDebuggingReplicator : public AActor
 	UPROPERTY(config)
 	FString DebugComponentHUDClassName;
 
+	UPROPERTY(config)
+	int32 MaxEQSQueries;
+
 	UPROPERTY(Replicated, Transient)
 	UGameplayDebuggingComponent* DebugComponent;
 
diff --git a/Engine/Source/Developer/GameplayDebugger/Private/GameplayDebuggingComponent.cpp b/Engine/Source/Developer/GameplayDebugger/Private/GameplayDebuggingComponent.cpp
--- a/Engine/Source/Developer/GameplayDebugger/Private/GameplayDebuggingComponent.cpp
+++ b/Engine/Source/Developer/GameplayDebugger/Private/GameplayDebuggingComponent.cpp
@@ -202,7 +202,6 @@ void UGameplayDebuggingComponent::GetLifetimeReplicatedProps( TArray< FLifetimeP
 	DOREPLIFETIME( UGameplayDebuggingComponent, TargetActor );
 
 	DOREPLIFETIME(UGameplayDebuggingComponent, EQSRepData);
-	DOREPLIFETIME(UGameplayDebuggingComponent, AllEQSName);
 #endif //!(UE_BUILD_SHIPPING || UE_BUILD_TEST)
 }
 
@@ -227,7 +226,6 @@ void UGameplayDebuggingComponent::SetActorToDebug(AActor* Actor)
 #if !(UE_BUILD_SHIPPING || UE_BUILD_TEST)
 	TargetActor = Actor;
 	EQSLocalData.Reset();
-	AllEQSName.Reset();
 #endif //!(UE_BUILD_SHIPPING || UE_BUILD_TEST)
 }
 
@@ -558,7 +556,7 @@ void UGameplayDebuggingComponent::ServerReplicateData(uint32 InMessage, uint32
 void UGameplayDebuggingComponent::OnChangeEQSQuery()
 {
 	AGameplayDebuggingReplicator* Replicator = Cast<AGameplayDebuggingReplicator>(GetOwner());
-	if (++CurrentEQSIndex >= AllEQSName.Num())
+	if (++CurrentEQSIndex >= EQSLocalData.Num())
 	{
 		CurrentEQSIndex = 0;
 	}
@@ -615,6 +613,7 @@ void UGameplayDebuggingComponent::CollectEQSData()
 	UWorld* World = GetWorld();
 	UEnvQueryManager* QueryManager = World ? UEnvQueryManager::GetCurrent(World) : NULL;
 	const AActor* Owner = GetSelectedActor();
+	AGameplayDebuggingReplicator* Replicator = Cast<AGameplayDebuggingReplicator>(GetOwner());
 
 	if (QueryManager == NULL || Owner == NULL)
 	{
@@ -628,36 +627,33 @@ void UGameplayDebuggingComponent::CollectEQSData()
 		const auto& AllControllerQueries = QueryManager->GetDebugger().GetAllQueriesForOwner(OwnerAsPawn->GetController());
 		AllQueries.Append(AllControllerQueries);
 	}
-	for (int32 Index = 0; Index < AllQueries.Num(); ++Index)
+	struct FEnvQueryInfoSort
 	{
-		EQSDebug::FQueryData* CurrentLocalData = NULL;
-		CachedQueryInstance = AllQueries[Index].Instance;
-		float CachedTimestamp = AllQueries[Index].Timestamp;
-		AllEQSName.AddUnique(CachedQueryInstance->QueryName);
-
-		 //find corresponding query
-		bool bSkipToNext = false;
-		for (int32 Idx = 0; Idx < EQSLocalData.Num(); ++Idx)
+		FORCEINLINE bool operator()(const FEQSDebugger::FEnvQueryInfo& A, const FEQSDebugger::FEnvQueryInfo& B) const
 		{
-			if (EQSLocalData[Idx].Name == CachedQueryInstance->QueryName)
-			{
-				if (EQSLocalData[Idx].Id == CachedQueryInstance->QueryID)
-				{
-					if (EQSLocalData[Idx].Timestamp == CachedTimestamp)
-					{
-						bSkipToNext = true;
-						break;
-					}
-				}
-				CurrentLocalData = &EQSLocalData[Index];;
-				break;
-			}
+			return (A.Timestamp < B.Timestamp);
 		}
+	};
+	TArray<FEQSDebugger::FEnvQueryInfo> QueriesToSort = AllQueries;
+	QueriesToSort.Sort(FEnvQueryInfoSort()); //sort queries by timestamp
+	QueriesToSort.SetNum(FMath::Min<int32>(Replicator->MaxEQSQueries, AllQueries.Num()));
 
-		if (bSkipToNext)
+	for (int32 Index = AllQueries.Num() - 1; Index >= 0; --Index)
+	{
+		auto &CurrentQuery = AllQueries[Index];
+		if (QueriesToSort.Find(CurrentQuery) == INDEX_NONE)
 		{
-			continue;
+			AllQueries.RemoveAt(Index);
 		}
+	}
+
+
+	EQSLocalData.Reset();
+	for (int32 Index = 0; Index < FMath::Min<int32>(Replicator->MaxEQSQueries, AllQueries.Num()); ++Index)
+	{
+		EQSDebug::FQueryData* CurrentLocalData = NULL;
+		CachedQueryInstance = AllQueries[Index].Instance;
+		const float CachedTimestamp = AllQueries[Index].Timestamp;
 
 		if (!CurrentLocalData)
 		{
@@ -1143,7 +1139,7 @@ FPrimitiveSceneProxy* UGameplayDebuggingComponent::CreateSceneProxy()
 #if USE_EQS_DEBUGGER
 	if (ShouldReplicateData(EAIDebugDrawDataView::EQS) && IsClientEQSSceneProxyEnabled())
 	{
-		const int32 EQSIndex = AllEQSName.Num() > 0 ? FMath::Clamp(CurrentEQSIndex, 0, AllEQSName.Num() - 1) : INDEX_NONE;
+		const int32 EQSIndex = EQSLocalData.Num() > 0 ? FMath::Clamp(CurrentEQSIndex, 0, EQSLocalData.Num() - 1) : INDEX_NONE;
 		if (EQSLocalData.IsValidIndex(EQSIndex))
 		{
 			CompositeProxy = CompositeProxy ? CompositeProxy : (new FDebugRenderSceneCompositeProxy(this));
diff --git a/Engine/Source/Developer/GameplayDebugger/Private/GameplayDebuggingHUDComponent.cpp b/Engine/Source/Developer/GameplayDebugger/Private/GameplayDebuggingHUDComponent.cpp
--- a/Engine/Source/Developer/GameplayDebugger/Private/GameplayDebuggingHUDComponent.cpp
+++ b/Engine/Source/Developer/GameplayDebugger/Private/GameplayDebuggingHUDComponent.cpp
@@ -347,22 +347,22 @@ void AGameplayDebuggingHUDComponent::DrawEQSData(APlayerController* PC, class UG
 #if !(UE_BUILD_SHIPPING || UE_BUILD_TEST) && WITH_EQS
 	PrintString(DefaultContext, TEXT("\n{green}EQS {white}[Use + key to switch query]\n"));
 
-	if (DebugComponent->AllEQSName.Num() == 0)
+	if (DebugComponent->EQSLocalData.Num() == 0)
 	{
 		return;
 	}
 
-	const int32 EQSIndex = DebugComponent->AllEQSName.Num() > 0 ? FMath::Clamp(DebugComponent->CurrentEQSIndex, 0, DebugComponent->AllEQSName.Num() - 1) : INDEX_NONE;
-	if (!DebugComponent->AllEQSName.IsValidIndex(EQSIndex))
+	const int32 EQSIndex = DebugComponent->EQSLocalData.Num() > 0 ? FMath::Clamp(DebugComponent->CurrentEQSIndex, 0, DebugComponent->EQSLocalData.Num() - 1) : INDEX_NONE;
+	if (!DebugComponent->EQSLocalData.IsValidIndex(EQSIndex))
 	{
 		return;
 	}
-	const FString SelecterEQSName = DebugComponent->AllEQSName[EQSIndex];
 
+	int32 Index = 0;
 	PrintString(DefaultContext, TEXT("{white}Queries: "));
 	for (auto CurrentQuery : DebugComponent->EQSLocalData)
 	{
-		if (SelecterEQSName == CurrentQuery.Name)
+		if (EQSIndex == Index)
 		{
 			PrintString(DefaultContext, FString::Printf(TEXT("{green}%s, "), *CurrentQuery.Name));
 		}
@@ -370,6 +370,7 @@ void AGameplayDebuggingHUDComponent::DrawEQSData(APlayerController* PC, class UG
 		{
 			PrintString(DefaultContext, FString::Printf(TEXT("{yellow}%s, "), *CurrentQuery.Name));
 		}
+		Index++;
 	}
 	PrintString(DefaultContext, TEXT("\n"));
 
@@ -515,7 +516,7 @@ void AGameplayDebuggingHUDComponent::DrawEQSItemDetails(int32 ItemIdx, class UGa
 	const float PosY = DefaultContext.CursorY + 1.0f;
 	float PosX = DefaultContext.CursorX;
 
-	const int32 EQSIndex = DebugComponent->AllEQSName.Num() > 0 ? FMath::Clamp(DebugComponent->CurrentEQSIndex, 0, DebugComponent->AllEQSName.Num() - 1) : INDEX_NONE;
+	const int32 EQSIndex = DebugComponent->EQSLocalData.Num() > 0 ? FMath::Clamp(DebugComponent->CurrentEQSIndex, 0, DebugComponent->EQSLocalData.Num() - 1) : INDEX_NONE;
 	auto& CurrentLocalData = DebugComponent->EQSLocalData[EQSIndex];
 	const EQSDebug::FItemData& ItemData = CurrentLocalData.Items[ItemIdx];
 
diff --git a/Engine/Source/Developer/GameplayDebugger/Private/GameplayDebuggingReplicator.cpp b/Engine/Source/Developer/GameplayDebugger/Private/GameplayDebuggingReplicator.cpp
--- a/Engine/Source/Developer/GameplayDebugger/Private/GameplayDebuggingReplicator.cpp
+++ b/Engine/Source/Developer/GameplayDebugger/Private/GameplayDebuggingReplicator.cpp
@@ -18,6 +18,7 @@ FOnSelectionChanged AGameplayDebuggingReplicator::OnSelectionChangedDelegate;
 
 AGameplayDebuggingReplicator::AGameplayDebuggingReplicator(const class FPostConstructInitializeProperties& PCIP)
 	: Super(PCIP)
+	, MaxEQSQueries(5)
 	, bIsGlobalInWorld(true)
 	, LastDrawAtFrame(0)
 	, PlayerControllersUpdateDelay(0)
