Fixed issue with auto activated GameplayDebuggeron clients #ue4

diff --git a/Engine/Source/Developer/GameplayDebugger/Classes/GameplayDebuggingReplicator.h b/Engine/Source/Developer/GameplayDebugger/Classes/GameplayDebuggingReplicator.h
--- a/Engine/Source/Developer/GameplayDebugger/Classes/GameplayDebuggingReplicator.h
+++ b/Engine/Source/Developer/GameplayDebugger/Classes/GameplayDebuggingReplicator.h
@@ -38,6 +38,9 @@ class GAMEPLAYDEBUGGER_API AGameplayDebuggingReplicator : public AActor
 	UPROPERTY(Replicated, Transient)
 	bool bIsGlobalInWorld;
 
+	UPROPERTY(ReplicatedUsing = OnRep_AutoActivate, Transient)
+	bool bAutoActivate;
+
 	UPROPERTY(Transient, EditAnywhere, Category = "DataView")
 	bool OverHead;
 	
@@ -87,6 +90,9 @@ class GAMEPLAYDEBUGGER_API AGameplayDebuggingReplicator : public AActor
 	void PostEditChangeProperty(FPropertyChangedEvent& PropertyChangedEvent) override;
 #endif
 
+	UFUNCTION()
+	virtual void OnRep_AutoActivate();
+
 	virtual class UNetConnection* GetNetConnection() override;
 
 	virtual bool IsNetRelevantFor(class APlayerController* RealViewer, AActor* Viewer, const FVector& SrcLocation) override;
diff --git a/Engine/Source/Developer/GameplayDebugger/Private/GameplayDebugger.cpp b/Engine/Source/Developer/GameplayDebugger/Private/GameplayDebugger.cpp
--- a/Engine/Source/Developer/GameplayDebugger/Private/GameplayDebugger.cpp
+++ b/Engine/Source/Developer/GameplayDebugger/Private/GameplayDebugger.cpp
@@ -147,6 +147,12 @@ bool FGameplayDebugger::CreateGameplayDebuggerForPlayerController(APlayerControl
 		DestActor->SetLocalPlayerOwner(PlayerController);
 		DestActor->SetReplicates(true);
 		DestActor->SetAsGlobalInWorld(false);
+		UEditorEngine* EEngine = Cast<UEditorEngine>(GEngine);
+		if (EEngine && EEngine->bIsSimulatingInEditor)
+		{
+			DestActor->CreateTool();
+			DestActor->EnableTool();
+		}
 		AddReplicator(World, DestActor);
 		return true;
 	}
@@ -292,6 +298,7 @@ bool FGameplayDebugger::Exec(UWorld* Inworld, const TCHAR* Cmd, FOutputDevice& A
 			if (!Replicator)
 			{
 				PC->ConsoleCommand("cheat EnableGDT");
+				bHandled = true;
 			}
 			else if (!Replicator->IsToolCreated())
 			{
@@ -329,6 +336,14 @@ bool FGameplayDebugger::Exec(UWorld* Inworld, const TCHAR* Cmd, FOutputDevice& A
 					bHandled = true;
 				}
 			}
+			else
+			{
+				if (Replicator)
+				{
+					Replicator->bAutoActivate = true;;
+					bHandled = true;
+				}
+			}
 		}
 	}
 	else if (FParse::Command(&Cmd, TEXT("RunEQS")))
diff --git a/Engine/Source/Developer/GameplayDebugger/Private/GameplayDebuggingComponent.cpp b/Engine/Source/Developer/GameplayDebugger/Private/GameplayDebuggingComponent.cpp
--- a/Engine/Source/Developer/GameplayDebugger/Private/GameplayDebuggingComponent.cpp
+++ b/Engine/Source/Developer/GameplayDebugger/Private/GameplayDebuggingComponent.cpp
@@ -1148,7 +1148,14 @@ FPrimitiveSceneProxy* UGameplayDebuggingComponent::CreateSceneProxy()
 		{
 			CompositeProxy = CompositeProxy ? CompositeProxy : (new FDebugRenderSceneCompositeProxy(this));
 			auto& CurrentLocalData = EQSLocalData[EQSIndex];
-			CompositeProxy->AddChild(new FEQSSceneProxy(this, TEXT("DebugAI"), false, CurrentLocalData.SolidSpheres, CurrentLocalData.Texts));
+			
+			FString ViewFlagName = TEXT("Game");
+			UEditorEngine* EEngine = Cast<UEditorEngine>(GEngine);
+			if (EEngine && EEngine->bIsSimulatingInEditor)
+			{
+				ViewFlagName = TEXT("DebugAI");
+			}
+			CompositeProxy->AddChild(new FEQSSceneProxy(this, ViewFlagName, false, CurrentLocalData.SolidSpheres, CurrentLocalData.Texts));
 		}
 	}
 #endif // USE_EQS_DEBUGGER
diff --git a/Engine/Source/Developer/GameplayDebugger/Private/GameplayDebuggingControllerComponent.cpp b/Engine/Source/Developer/GameplayDebugger/Private/GameplayDebuggingControllerComponent.cpp
--- a/Engine/Source/Developer/GameplayDebugger/Private/GameplayDebuggingControllerComponent.cpp
+++ b/Engine/Source/Developer/GameplayDebugger/Private/GameplayDebuggingControllerComponent.cpp
@@ -156,7 +156,6 @@ void UGameplayDebuggingControllerComponent::OnActivationKeyPressed()
 			BindAIDebugViewKeys();
 			GetDebuggingReplicator()->EnableDraw(true);
 			GetDebuggingReplicator()->ServerReplicateMessage(NULL, EDebugComponentMessage::ActivateReplication, EAIDebugDrawDataView::Empty);
-			PlayerOwner->ConsoleCommand(TEXT("ShowFlag.DebugAI 1"), false);
 		}
 
 		ControlKeyPressedTime = GetWorld()->GetTimeSeconds();
@@ -205,10 +204,6 @@ void UGameplayDebuggingControllerComponent::CloseDebugTool()
 		GetDebuggingReplicator()->EnableDraw(false);
 		GetDebuggingReplicator()->ServerReplicateMessage(NULL, EDebugComponentMessage::DeactivateReplilcation, EAIDebugDrawDataView::Empty);
 		bToolActivated = false;
-		if (PlayerOwner.IsValid())
-		{
-			PlayerOwner->ConsoleCommand(TEXT("ShowFlag.DebugAI 0"), false);
-		}
 	}
 }
 
diff --git a/Engine/Source/Developer/GameplayDebugger/Private/GameplayDebuggingReplicator.cpp b/Engine/Source/Developer/GameplayDebugger/Private/GameplayDebuggingReplicator.cpp
--- a/Engine/Source/Developer/GameplayDebugger/Private/GameplayDebuggingReplicator.cpp
+++ b/Engine/Source/Developer/GameplayDebugger/Private/GameplayDebuggingReplicator.cpp
@@ -245,7 +245,10 @@ void AGameplayDebuggingReplicator::BeginPlay()
 			}
 		}
 	}
+}
 
+void AGameplayDebuggingReplicator::OnRep_AutoActivate()
+{
 	// we are already replicated so let's activate tool
 	if (GetWorld() && GetNetMode() == ENetMode::NM_Client && !IsToolCreated() && !IsGlobalInWorld())
 	{
