fixed TTP #348115 "ENGINE: Gameplay debug tool input mapping is incorrect"

diff --git a/Engine/Source/Developer/GameplayDebugger/Private/GameplayDebuggingHUDComponent.cpp b/Engine/Source/Developer/GameplayDebugger/Private/GameplayDebuggingHUDComponent.cpp
--- a/Engine/Source/Developer/GameplayDebugger/Private/GameplayDebuggingHUDComponent.cpp
+++ b/Engine/Source/Developer/GameplayDebugger/Private/GameplayDebuggingHUDComponent.cpp
@@ -134,7 +134,7 @@ void AGameplayDebuggingHUDComponent::DrawMenu(const float X, const float Y, clas
 		CategoriesWidth.AddZeroed(Categories.Num());
 		float TotalWidth = 0.0f, MaxHeight = 0.0f;
 
-		FString HeaderDesc(TEXT("Tap [\"] to close, use Numpad to toggle categories"));
+		FString HeaderDesc(TEXT("Tap ['] to close, use Numpad to toggle categories"));
 		float HeaderWidth = 0.0f;
 		CalulateStringSize(DefaultContext, DefaultContext.Font, HeaderDesc, HeaderWidth, MaxHeight);
 
@@ -170,7 +170,7 @@ void AGameplayDebuggingHUDComponent::DrawMenu(const float X, const float Y, clas
 
 	if ((!DebugComponent || !DebugComponent->GetSelectedActor()) && GetWorld()->GetNetMode() == NM_Client)
 	{
-		PrintString(DefaultContext, "\n{red}No Pawn selected - waiting for data to replicate from server. {green}Press and hold \" to select Pawn \n");
+		PrintString(DefaultContext, "\n{red}No Pawn selected - waiting for data to replicate from server. {green}Press and hold ' to select Pawn \n");
 	}
 
 	DefaultContext.CursorX = OldX;
