Fixed CIS error inside GameplayDebugger #ue4

diff --git a/Engine/Source/Developer/GameplayDebugger/Private/GameplayDebugger.cpp b/Engine/Source/Developer/GameplayDebugger/Private/GameplayDebugger.cpp
--- a/Engine/Source/Developer/GameplayDebugger/Private/GameplayDebugger.cpp
+++ b/Engine/Source/Developer/GameplayDebugger/Private/GameplayDebugger.cpp
@@ -7,6 +7,9 @@
 #include "GameplayDebuggingReplicator.h"
 #include "GameplayDebuggingControllerComponent.h"
 #include "AISystem.h"
+#if WITH_EDITOR
+#include "Editor/EditorEngine.h"
+#endif // WITH_EDITOR
 
 uint32 FGameplayDebuggerSettings::ShowFlagIndex = 0;
 
@@ -147,12 +150,14 @@ bool FGameplayDebugger::CreateGameplayDebuggerForPlayerController(APlayerControl
 		DestActor->SetLocalPlayerOwner(PlayerController);
 		DestActor->SetReplicates(true);
 		DestActor->SetAsGlobalInWorld(false);
+#if WITH_EDITOR
 		UEditorEngine* EEngine = Cast<UEditorEngine>(GEngine);
 		if (EEngine && EEngine->bIsSimulatingInEditor)
 		{
 			DestActor->CreateTool();
 			DestActor->EnableTool();
 		}
+#endif
 		AddReplicator(World, DestActor);
 		return true;
 	}
