Forgot about one file with CIS fixes. #ue4

diff --git a/Engine/Source/Developer/GameplayDebugger/Private/GameplayDebuggingComponent.cpp b/Engine/Source/Developer/GameplayDebugger/Private/GameplayDebuggingComponent.cpp
--- a/Engine/Source/Developer/GameplayDebugger/Private/GameplayDebuggingComponent.cpp
+++ b/Engine/Source/Developer/GameplayDebugger/Private/GameplayDebuggingComponent.cpp
@@ -1150,11 +1150,13 @@ FPrimitiveSceneProxy* UGameplayDebuggingComponent::CreateSceneProxy()
 			auto& CurrentLocalData = EQSLocalData[EQSIndex];
 			
 			FString ViewFlagName = TEXT("Game");
+#if WITH_EDITOR
 			UEditorEngine* EEngine = Cast<UEditorEngine>(GEngine);
 			if (EEngine && EEngine->bIsSimulatingInEditor)
 			{
 				ViewFlagName = TEXT("DebugAI");
 			}
+#endif
 			CompositeProxy->AddChild(new FEQSSceneProxy(this, ViewFlagName, false, CurrentLocalData.SolidSpheres, CurrentLocalData.Texts));
 		}
 	}
