A few more TSubobjectPtrs removed

diff --git a/Engine/Plugins/Runtime/CableComponent/Source/CableComponent/Classes/CableActor.h b/Engine/Plugins/Runtime/CableComponent/Source/CableComponent/Classes/CableActor.h
--- a/Engine/Plugins/Runtime/CableComponent/Source/CableComponent/Classes/CableActor.h
+++ b/Engine/Plugins/Runtime/CableComponent/Source/CableComponent/Classes/CableActor.h
@@ -12,5 +12,5 @@ class ACableActor : public AActor
 
 	/** Cable component that performs simulation and rendering */
 	UPROPERTY(Category=Cable, VisibleAnywhere, BlueprintReadOnly)
-	TSubobjectPtr<class UCableComponent> CableComponent;
+	class UCableComponent* CableComponent;
 };
diff --git a/Engine/Source/Developer/GameplayDebugger/Private/GameplayDebuggingReplicator.cpp b/Engine/Source/Developer/GameplayDebugger/Private/GameplayDebuggingReplicator.cpp
--- a/Engine/Source/Developer/GameplayDebugger/Private/GameplayDebuggingReplicator.cpp
+++ b/Engine/Source/Developer/GameplayDebugger/Private/GameplayDebuggingReplicator.cpp
@@ -41,7 +41,7 @@ AGameplayDebuggingReplicator::AGameplayDebuggingReplicator(const FObjectInitiali
 	PrimaryActorTick.bCanEverTick = true;
 	PrimaryActorTick.bStartWithTickEnabled = false;
 	
-	TSubobjectPtr<USceneComponent> SceneComponent = ObjectInitializer.CreateDefaultSubobject<USceneComponent>(this, TEXT("SceneComponent"));
+	USceneComponent* SceneComponent = ObjectInitializer.CreateDefaultSubobject<USceneComponent>(this, TEXT("SceneComponent"));
 	RootComponent = SceneComponent;
 
 #if WITH_EDITOR
