Fixed GameplayDebugger activation with keyboard shortcut #ue4

diff --git a/Engine/Config/BaseInput.ini b/Engine/Config/BaseInput.ini
--- a/Engine/Config/BaseInput.ini
+++ b/Engine/Config/BaseInput.ini
@@ -42,7 +42,7 @@ bRequireCtrlToNavigateAutoComplete=False
 +DebugExecBindings=(Key=Enter,Command="TOGGLE_FULLSCREEN",Alt=True)
 +DebugExecBindings=(Key=Tab,Command="FocusNextPIEWindow",Control=True)
 +DebugExecBindings=(Key=Tab,Command="FocusLastPIEWindow",Control=True,Shift=True)
-+DebugExecBindings=(Key=Quote,Command="EnableGDT",Control=False,Shift=False)
++DebugExecBindings=(Key=Apostrophe,Command="EnableGDT",Control=False,Shift=False)
 
 [/Script/Engine.Console]
 HistoryBot=-1
diff --git a/Engine/Source/Developer/GameplayDebugger/Private/GameplayDebuggingControllerComponent.cpp b/Engine/Source/Developer/GameplayDebugger/Private/GameplayDebuggingControllerComponent.cpp
--- a/Engine/Source/Developer/GameplayDebugger/Private/GameplayDebuggingControllerComponent.cpp
+++ b/Engine/Source/Developer/GameplayDebugger/Private/GameplayDebuggingControllerComponent.cpp
@@ -126,7 +126,7 @@ void UGameplayDebuggingControllerComponent::BindActivationKeys()
 	if (PlayerOwner.IsValid() && PlayerOwner->InputComponent && PlayerOwner->PlayerInput)
 	{
 		// find current activation key used for 'EnableGDT' binding
-		FInputChord ActivationKey(EKeys::Quote, false, false, false, false);
+		FInputChord ActivationKey(EKeys::Apostrophe, false, false, false, false);
 		for (uint32 BindIndex = 0; BindIndex < (uint32)PlayerOwner->PlayerInput->DebugExecBindings.Num(); BindIndex++)
 		{
 			if (PlayerOwner->PlayerInput->DebugExecBindings[BindIndex].Command == TEXT("EnableGDT"))
