Fixed keyboard shortcut display name for GameplayDebugger - it should work Ok with any keyboard binding now.

diff --git a/Engine/Source/Developer/GameplayDebugger/Classes/GameplayDebuggingControllerComponent.h b/Engine/Source/Developer/GameplayDebugger/Classes/GameplayDebuggingControllerComponent.h
--- a/Engine/Source/Developer/GameplayDebugger/Classes/GameplayDebuggingControllerComponent.h
+++ b/Engine/Source/Developer/GameplayDebugger/Classes/GameplayDebuggingControllerComponent.h
@@ -36,6 +36,7 @@ class GAMEPLAYDEBUGGER_API UGameplayDebuggingControllerComponent : public UActor
 		EnableActiveView(View, IsViewActive(View) ? false : true);
 	}
 
+	const FInputChord& GetActivationKey() const { return ActivationKey; }
 
 	/** periodic update of navmesh data */
 	void UpdateNavMeshTimer();
@@ -76,6 +77,7 @@ protected:
 
 	TWeakObjectPtr<APlayerController> PlayerOwner;
 
+	FInputChord ActivationKey;
 	const float KeyPressActivationTime;
 	float ControlKeyPressedTime;
 	float PlayersComponentRequestTime;
diff --git a/Engine/Source/Developer/GameplayDebugger/Private/GameplayDebuggingControllerComponent.cpp b/Engine/Source/Developer/GameplayDebugger/Private/GameplayDebuggingControllerComponent.cpp
--- a/Engine/Source/Developer/GameplayDebugger/Private/GameplayDebuggingControllerComponent.cpp
+++ b/Engine/Source/Developer/GameplayDebugger/Private/GameplayDebuggingControllerComponent.cpp
@@ -29,6 +29,7 @@ UGameplayDebuggingControllerComponent::UGameplayDebuggingControllerComponent(con
 	bWaitingForOwnersComponent = false;
 
 	ControlKeyPressedTime = 0;
+	ActivationKey = FInputChord(EKeys::Apostrophe, false, false, false, false);
 }
 
 void UGameplayDebuggingControllerComponent::OnRegister()
@@ -126,7 +127,6 @@ void UGameplayDebuggingControllerComponent::BindActivationKeys()
 	if (PlayerOwner.IsValid() && PlayerOwner->InputComponent && PlayerOwner->PlayerInput)
 	{
 		// find current activation key used for 'EnableGDT' binding
-		FInputChord ActivationKey(EKeys::Apostrophe, false, false, false, false);
 		for (uint32 BindIndex = 0; BindIndex < (uint32)PlayerOwner->PlayerInput->DebugExecBindings.Num(); BindIndex++)
 		{
 			if (PlayerOwner->PlayerInput->DebugExecBindings[BindIndex].Command == TEXT("EnableGDT"))
diff --git a/Engine/Source/Developer/GameplayDebugger/Private/GameplayDebuggingHUDComponent.cpp b/Engine/Source/Developer/GameplayDebugger/Private/GameplayDebuggingHUDComponent.cpp
--- a/Engine/Source/Developer/GameplayDebugger/Private/GameplayDebuggingHUDComponent.cpp
+++ b/Engine/Source/Developer/GameplayDebugger/Private/GameplayDebuggingHUDComponent.cpp
@@ -134,7 +134,20 @@ void AGameplayDebuggingHUDComponent::DrawMenu(const float X, const float Y, clas
 		CategoriesWidth.AddZeroed(Categories.Num());
 		float TotalWidth = 0.0f, MaxHeight = 0.0f;
 
-		FString HeaderDesc(TEXT("Tap ['] to close, use Numpad to toggle categories"));
+		FString ActivationKeyDisplayName = TEXT("'");
+		FString ActivationKeyName = TEXT("Apostrophe");
+
+		APlayerController* const MyPC = Cast<APlayerController>(PlayerOwner);
+		UGameplayDebuggingControllerComponent*  GDC = GetDebuggingReplicator()->FindComponentByClass<UGameplayDebuggingControllerComponent>();
+		if (GDC)
+		{
+			ActivationKeyDisplayName = GDC->GetActivationKey().Key.GetDisplayName().ToString();
+			ActivationKeyName = GDC->GetActivationKey().Key.GetFName().ToString();
+		}
+
+		const FString KeyDesc = ActivationKeyName != ActivationKeyDisplayName ? FString::Printf(TEXT("(%s key)"), *ActivationKeyName) : TEXT("");
+		const FString HeaderDesc = FString::Printf(TEXT("Tap %s %s to close, use Numpad to toggle categories"), *ActivationKeyDisplayName, *KeyDesc);
+
 		float HeaderWidth = 0.0f;
 		CalulateStringSize(DefaultContext, DefaultContext.Font, HeaderDesc, HeaderWidth, MaxHeight);
 
