File missing in prev submission, CIS fix #UE4

diff --git a/Engine/Source/Developer/GameplayDebugger/Private/GameplayDebuggingComponent.cpp b/Engine/Source/Developer/GameplayDebugger/Private/GameplayDebuggingComponent.cpp
--- a/Engine/Source/Developer/GameplayDebugger/Private/GameplayDebuggingComponent.cpp
+++ b/Engine/Source/Developer/GameplayDebugger/Private/GameplayDebuggingComponent.cpp
@@ -858,11 +858,8 @@ ARecastNavMesh* UGameplayDebuggingComponent::GetNavData()
 	APawn* TargetPawn = Cast<APawn>(TargetActor);
 	if (TargetPawn != NULL)
 	{
-		const FNavAgentProperties* NavAgentProperties = TargetPawn->GetNavAgentProperties();
-		if (NavAgentProperties != NULL)
-		{
-			return Cast<ARecastNavMesh>(NavSys->GetNavDataForProps(*NavAgentProperties));
-		}
+		const FNavAgentProperties& NavAgentProperties = TargetPawn->GetNavAgentProperties();
+		return Cast<ARecastNavMesh>(NavSys->GetNavDataForProps(NavAgentProperties));
 	}
 
 	// If it wasn't found, just get the main nav-mesh data.
