Prevented to create DebuggingReplicator from GameplayDebugger in non game world or in pending kill worlds. #ue4

diff --git a/Engine/Source/Developer/GameplayDebugger/Private/GameplayDebugger.cpp b/Engine/Source/Developer/GameplayDebugger/Private/GameplayDebugger.cpp
--- a/Engine/Source/Developer/GameplayDebugger/Private/GameplayDebugger.cpp
+++ b/Engine/Source/Developer/GameplayDebugger/Private/GameplayDebugger.cpp
@@ -189,6 +189,11 @@ void FGameplayDebugger::WorldAdded(UWorld* InWorld)
 		return;
 	}
 
+	if (InWorld == NULL || InWorld->IsPendingKill() || InWorld->IsGameWorld() == false)
+	{
+		return;
+	}
+
 	for (auto It = GetAllReplicators(InWorld).CreateConstIterator(); It; ++It)
 	{
 		TWeakObjectPtr<AGameplayDebuggingReplicator> Replicator = *It;
