Fix for UE-4866 "GameplayDebugger is not frawing not-stringpulled paths correctly"

diff --git a/Engine/Source/Developer/GameplayDebugger/Classes/GameplayDebuggingComponent.h b/Engine/Source/Developer/GameplayDebugger/Classes/GameplayDebuggingComponent.h
--- a/Engine/Source/Developer/GameplayDebugger/Classes/GameplayDebuggingComponent.h
+++ b/Engine/Source/Developer/GameplayDebugger/Classes/GameplayDebuggingComponent.h
@@ -76,6 +76,12 @@ class GAMEPLAYDEBUGGER_API UGameplayDebuggingComponent : public UPrimitiveCompon
 	UPROPERTY(Replicated)
 	TArray<FVector> PathPoints;
 
+	/** Begin path replication data */
+	UPROPERTY(ReplicatedUsing = OnRep_PathCorridorData)
+	TArray<uint8> PathCorridorData;
+	
+	TArray<TArray<FVector> > PathCorridorPolygons;
+
 	UPROPERTY(Replicated)
 	FString PathErrorString;
 	/** End path replication data*/
@@ -125,6 +131,9 @@ class GAMEPLAYDEBUGGER_API UGameplayDebuggingComponent : public UPrimitiveCompon
 
 	void PrepareNavMeshData(struct FNavMeshSceneProxyData*) const;
 
+	UFUNCTION()
+	virtual void OnRep_PathCorridorData();
+
 	virtual void Activate(bool bReset=false) override;
 	virtual void Deactivate() override;
 
diff --git a/Engine/Source/Developer/GameplayDebugger/Private/GameplayDebuggingComponent.cpp b/Engine/Source/Developer/GameplayDebugger/Private/GameplayDebuggingComponent.cpp
--- a/Engine/Source/Developer/GameplayDebugger/Private/GameplayDebuggingComponent.cpp
+++ b/Engine/Source/Developer/GameplayDebugger/Private/GameplayDebuggingComponent.cpp
@@ -18,6 +18,11 @@
 #include "GameFramework/PlayerState.h"
 #include "Engine/Channel.h"
 
+#if WITH_EDITOR
+#include "Editor/EditorEngine.h"
+#include "GeomTools.h"
+#endif // WITH_EDITOR
+
 #include "GameplayDebuggingComponent.h"
 
 DEFINE_LOG_CATEGORY(LogGDT);
@@ -153,6 +158,11 @@ UGameplayDebuggingComponent::UGameplayDebuggingComponent(const FObjectInitialize
 	if (Replicator)
 	{
 		Replicator->OnChangeEQSQuery.AddUObject(this, &UGameplayDebuggingComponent::OnChangeEQSQuery);
+		FGameplayDebuggerSettings Settings = GameplayDebuggerSettings(Replicator);
+		for (int32 Index = EAIDebugDrawDataView::Empty + 1; Index < EAIDebugDrawDataView::MAX; ++Index)
+		{
+			ReplicateViewDataCounters[Index] = Settings.CheckFlag(Index);
+		}
 	}
 
 }
@@ -205,6 +215,7 @@ void UGameplayDebuggingComponent::GetLifetimeReplicatedProps( TArray< FLifetimeP
 
 	DOREPLIFETIME(UGameplayDebuggingComponent, SensingComponentLocation);
 
+	DOREPLIFETIME(UGameplayDebuggingComponent, PathCorridorData);
 #endif //!(UE_BUILD_SHIPPING || UE_BUILD_TEST)
 }
 
@@ -477,6 +488,8 @@ void UGameplayDebuggingComponent::CollectPathData()
 #if !(UE_BUILD_SHIPPING || UE_BUILD_TEST)
 	APawn* MyPawn = Cast<APawn>(GetSelectedActor());
 	PathErrorString.Empty();
+
+	bool bRefreshRendering = false;
 	if (AAIController *MyController = Cast<AAIController>(MyPawn->Controller))
 	{
 		if ( MyController->GetNavComponent() && MyController->GetNavComponent()->HasValidPath())
@@ -484,23 +497,89 @@ void UGameplayDebuggingComponent::CollectPathData()
 			const FNavPathSharedPtr& NewPath = MyController->GetNavComponent()->GetPath();
 			if (!CurrentPath.HasSameObject(NewPath.Get()))
 			{
+				FVisualLogEntry Snapshot;
+				NewPath->DescribeSelfToVisLog(&Snapshot);
+				PathCorridorPolygons.Reserve(Snapshot.ElementsToDraw.Num());
+				for (auto& CurrentShape : Snapshot.ElementsToDraw)
+				{
+					if (CurrentShape.GetType() == EVisualLoggerShapeElement::Polygon)
+					{
+						PathCorridorPolygons.Add(CurrentShape.Points);
+					}
+				}
+
 				PathPoints.Reset();
 				for (int32 Index=0; Index < NewPath->GetPathPoints().Num(); ++Index)
 				{
 					PathPoints.Add(NewPath->GetPathPoints()[Index].Location);
 				}
 				CurrentPath = NewPath;
+
+				if (PathCorridorPolygons.Num() && World && World->GetNetMode() != NM_Client)
+				{
+					PathCorridorData.Reset();
+					TArray<uint8> HelpBuffer;
+					FMemoryWriter ArWriter(HelpBuffer);
+					int32 NumPolygons = PathCorridorPolygons.Num();
+					ArWriter << NumPolygons;
+					for (const auto& CurrentPoly : PathCorridorPolygons)
+					{
+						int32 NumVerts = CurrentPoly.Num();
+						ArWriter << NumVerts;
+						for (auto Vert : CurrentPoly)
+						{
+							ArWriter << Vert;
+						}
+					}
+				}
+				bRefreshRendering = true;
 			}
 		}
 		else
 		{
 			CurrentPath = NULL;
 			PathPoints.Reset();
+			PathCorridorPolygons.Reset();
 		}
 	}
+
+	if (bRefreshRendering && World && World->GetNetMode() != NM_DedicatedServer)
+	{
+		UpdateBounds();
+		MarkRenderStateDirty();
+	}
 #endif //!(UE_BUILD_SHIPPING || UE_BUILD_TEST)
 }
 
+void UGameplayDebuggingComponent::OnRep_PathCorridorData()
+{
+#if !(UE_BUILD_SHIPPING || UE_BUILD_TEST)
+	if (World && World->GetNetMode() != NM_DedicatedServer)
+	{
+		TArray<uint8> HelpBuffer;
+		FMemoryReader ArReader(HelpBuffer);
+		int32 NumPolygons = 0;
+		ArReader << NumPolygons;
+		PathCorridorPolygons.Reset(NumPolygons);
+		for (int32 PolyIndex = 0; PolyIndex < NumPolygons; ++PolyIndex)
+		{
+			TArray<FVector> CurrentPoly;
+			int32 NumVerts = 0;
+			ArReader << NumVerts;
+			for (int32 VertIndex = 0; VertIndex < NumVerts; ++VertIndex)
+			{
+				FVector CurrentVertex;
+				ArReader << CurrentVertex;
+				CurrentPoly.Add(CurrentVertex);
+			}
+		}
+		UpdateBounds();
+		MarkRenderStateDirty();
+	}
+#endif
+}
+
+
 void UGameplayDebuggingComponent::EnableDebugDraw(bool bEnable, bool InFocusedComponent)
 {
 #if !(UE_BUILD_SHIPPING || UE_BUILD_TEST) 
@@ -1111,6 +1190,29 @@ void UGameplayDebuggingComponent::PrepareNavMeshData(struct FNavMeshSceneProxyDa
 //----------------------------------------------------------------------//
 // rendering
 //----------------------------------------------------------------------//
+class FPathDebugRenderSceneProxy : public FDebugRenderSceneProxy
+{
+public:
+	FPathDebugRenderSceneProxy(const UPrimitiveComponent* InComponent, const FString &InViewFlagName)
+		: FDebugRenderSceneProxy(InComponent)
+	{
+		DrawType = FDebugRenderSceneProxy::SolidAndWireMeshes;
+		ViewFlagName = InViewFlagName;
+		ViewFlagIndex = uint32(FEngineShowFlags::FindIndexByName(*ViewFlagName));
+	}
+
+
+	FPrimitiveViewRelevance GetViewRelevance(const FSceneView* View) override
+	{
+		FPrimitiveViewRelevance Result;
+		Result.bDrawRelevance = View->Family->EngineShowFlags.GetSingleFlag(ViewFlagIndex);// IsShown(View);
+		Result.bDynamicRelevance = true;
+		Result.bNormalTranslucencyRelevance = IsShown(View);
+		return Result;
+	}
+
+};
+
 FPrimitiveSceneProxy* UGameplayDebuggingComponent::CreateSceneProxy()
 {
 	FDebugRenderSceneCompositeProxy* CompositeProxy = NULL;
@@ -1162,6 +1264,70 @@ FPrimitiveSceneProxy* UGameplayDebuggingComponent::CreateSceneProxy()
 	}
 #endif // USE_EQS_DEBUGGER
 
+	const bool bDrawFullData = Replicator->GetSelectedActorToDebug() == GetSelectedActor();
+	if (bDrawFullData && ShouldReplicateData(EAIDebugDrawDataView::Basic))
+	{
+		CompositeProxy = CompositeProxy ? CompositeProxy : (new FDebugRenderSceneCompositeProxy(this));
+
+		TArray<FDebugRenderSceneProxy::FMesh> Meshes;
+		TArray<FDebugRenderSceneProxy::FDebugLine> Lines;
+		for (auto& CurrentPoly : PathCorridorPolygons)
+		{
+#if WITH_EDITOR
+			FClipSMPolygon InPoly(CurrentPoly.Num());
+			for (int32 Index = 0; Index < CurrentPoly.Num(); Index++)
+			{
+				FClipSMVertex v1;
+				v1.Pos = CurrentPoly[Index];
+				InPoly.Vertices.Add(v1);
+			}
+
+			TArray<FClipSMTriangle> OutTris;
+			if (TriangulatePoly(OutTris, InPoly, false))
+			{
+				int32 LastIndex = 0;
+				FDebugRenderSceneProxy::FMesh TestMesh;
+				TestMesh.Color = FColor::Green;
+
+				RemoveRedundantTriangles(OutTris);
+				for (const auto& CurrentTri : OutTris)
+				{
+					TestMesh.Vertices.Add(FDynamicMeshVertex(CurrentTri.Vertices[0].Pos));
+					TestMesh.Vertices.Add(FDynamicMeshVertex(CurrentTri.Vertices[1].Pos));
+					TestMesh.Vertices.Add(FDynamicMeshVertex(CurrentTri.Vertices[2].Pos));
+					TestMesh.Indices.Add(LastIndex++);
+					TestMesh.Indices.Add(LastIndex++);
+					TestMesh.Indices.Add(LastIndex++);
+				}
+				Meshes.Add(TestMesh);
+			}
+#endif
+
+			for (int32 VIdx = 0; VIdx < CurrentPoly.Num(); VIdx++)
+			{
+				Lines.Add(FDebugRenderSceneProxy::FDebugLine(
+					CurrentPoly[VIdx],
+					CurrentPoly[(VIdx + 1) % CurrentPoly.Num()],
+					FColor::Cyan,
+					2)
+				);
+			}
+		}
+
+		FString ViewFlagName = TEXT("Game");
+#if WITH_EDITOR
+		UEditorEngine* EEngine = Cast<UEditorEngine>(GEngine);
+		if (EEngine && EEngine->bIsSimulatingInEditor)
+		{
+			ViewFlagName = TEXT("DebugAI");
+		}
+#endif
+		FPathDebugRenderSceneProxy *DebugSceneProxy = new FPathDebugRenderSceneProxy(this, ViewFlagName);
+		DebugSceneProxy->Lines = Lines;
+		DebugSceneProxy->Meshes = Meshes;
+		CompositeProxy->AddChild(DebugSceneProxy);
+	}
+
 	return CompositeProxy;
 }
 
@@ -1178,7 +1344,7 @@ FBoxSphereBounds UGameplayDebuggingComponent::CalcBounds(const FTransform& Local
 #endif
 
 #if USE_EQS_DEBUGGER
-	if (EQSRepData.Num() && ShouldReplicateData(EAIDebugDrawDataView::EQS))
+	if ((EQSRepData.Num() && ShouldReplicateData(EAIDebugDrawDataView::EQS)) || PathCorridorPolygons.Num())
 	{
 		MyBounds = FBox(FVector(-HALF_WORLD_MAX, -HALF_WORLD_MAX, -HALF_WORLD_MAX), FVector(HALF_WORLD_MAX, HALF_WORLD_MAX, HALF_WORLD_MAX));
 	}
diff --git a/Engine/Source/Developer/GameplayDebugger/Private/GameplayDebuggingReplicator.cpp b/Engine/Source/Developer/GameplayDebugger/Private/GameplayDebuggingReplicator.cpp
--- a/Engine/Source/Developer/GameplayDebugger/Private/GameplayDebuggingReplicator.cpp
+++ b/Engine/Source/Developer/GameplayDebugger/Private/GameplayDebuggingReplicator.cpp
@@ -210,6 +210,18 @@ void AGameplayDebuggingReplicator::BeginPlay()
 		if (DebugComponent)
 		{
 			DebugComponent->ServerReplicateData(EDebugComponentMessage::ActivateReplication, EAIDebugDrawDataView::Empty);
+
+			FGameplayDebuggerSettings Settings = GameplayDebuggerSettings(this);
+			DebugComponent->ServerReplicateData(Settings.CheckFlag(EAIDebugDrawDataView::OverHead) ? EDebugComponentMessage::ActivateDataView : EDebugComponentMessage::DeactivateDataView, EAIDebugDrawDataView::OverHead);
+			DebugComponent->ServerReplicateData(Settings.CheckFlag(EAIDebugDrawDataView::Basic) ? EDebugComponentMessage::ActivateDataView : EDebugComponentMessage::DeactivateDataView, EAIDebugDrawDataView::Basic);
+			DebugComponent->ServerReplicateData(Settings.CheckFlag(EAIDebugDrawDataView::BehaviorTree) ? EDebugComponentMessage::ActivateDataView : EDebugComponentMessage::DeactivateDataView, EAIDebugDrawDataView::BehaviorTree);
+			DebugComponent->ServerReplicateData(Settings.CheckFlag(EAIDebugDrawDataView::EQS) ? EDebugComponentMessage::ActivateDataView : EDebugComponentMessage::DeactivateDataView, EAIDebugDrawDataView::EQS);
+			DebugComponent->ServerReplicateData(Settings.CheckFlag(EAIDebugDrawDataView::Perception) ? EDebugComponentMessage::ActivateDataView : EDebugComponentMessage::DeactivateDataView, EAIDebugDrawDataView::Perception);
+			DebugComponent->ServerReplicateData(Settings.CheckFlag(EAIDebugDrawDataView::GameView1) ? EDebugComponentMessage::ActivateDataView : EDebugComponentMessage::DeactivateDataView, EAIDebugDrawDataView::GameView1);
+			DebugComponent->ServerReplicateData(Settings.CheckFlag(EAIDebugDrawDataView::GameView2) ? EDebugComponentMessage::ActivateDataView : EDebugComponentMessage::DeactivateDataView, EAIDebugDrawDataView::GameView2);
+			DebugComponent->ServerReplicateData(Settings.CheckFlag(EAIDebugDrawDataView::GameView3) ? EDebugComponentMessage::ActivateDataView : EDebugComponentMessage::DeactivateDataView, EAIDebugDrawDataView::GameView3);
+			DebugComponent->ServerReplicateData(Settings.CheckFlag(EAIDebugDrawDataView::GameView4) ? EDebugComponentMessage::ActivateDataView : EDebugComponentMessage::DeactivateDataView, EAIDebugDrawDataView::GameView4);
+			DebugComponent->ServerReplicateData(Settings.CheckFlag(EAIDebugDrawDataView::GameView5) ? EDebugComponentMessage::ActivateDataView : EDebugComponentMessage::DeactivateDataView, EAIDebugDrawDataView::GameView5);
 		}
 	}
 	else
diff --git a/Engine/Source/Runtime/Engine/Private/DebugRenderSceneProxy.cpp b/Engine/Source/Runtime/Engine/Private/DebugRenderSceneProxy.cpp
--- a/Engine/Source/Runtime/Engine/Private/DebugRenderSceneProxy.cpp
+++ b/Engine/Source/Runtime/Engine/Private/DebugRenderSceneProxy.cpp
@@ -16,6 +16,7 @@
 FDebugRenderSceneProxy::FDebugRenderSceneProxy(const UPrimitiveComponent* InComponent) 
 	: FPrimitiveSceneProxy(InComponent)
 	, ViewFlagIndex(uint32(FEngineShowFlags::FindIndexByName(TEXT("Game"))))
+	, ViewFlagName(TEXT("Game"))
 	, TextWithoutShadowDistance(1500)
 	, DrawType(WireMesh)
 {
