Added AI Tools section to Editor Settings
Added GameplayDebugger settings to editor (AI Tools section) - it's usefull feature to controll default properties.
Fixed LogVisualizer as experimental menu entry in UE4 Editor

diff --git a/Engine/Source/Developer/GameplayDebugger/Classes/GameplayDebuggerSettings.h b/Engine/Source/Developer/GameplayDebugger/Classes/GameplayDebuggerSettings.h
new file mode 100644
--- /dev/null
+++ b/Engine/Source/Developer/GameplayDebugger/Classes/GameplayDebuggerSettings.h
@@ -0,0 +1,102 @@
+// Copyright 1998-2014 Epic Games, Inc. All Rights Reserved.
+
+/*=============================================================================
+GameplayDebuggerSettings.h: Declares the UGameplayDebuggerSettings class.
+=============================================================================*/
+#pragma once
+
+
+#include "GameplayDebuggerSettings.generated.h"
+
+USTRUCT()
+struct FGDTCustomViewNames
+{
+	GENERATED_USTRUCT_BODY()
+
+	UPROPERTY(EditAnywhere, config, Category = "Settings")
+	FString GameView1;
+
+	UPROPERTY(EditAnywhere, config, Category = "Settings")
+	FString GameView2;
+
+	UPROPERTY(EditAnywhere, config, Category = "Settings")
+	FString GameView3;
+
+	UPROPERTY(EditAnywhere, config, Category = "Settings")
+	FString GameView4;
+
+	UPROPERTY(EditAnywhere, config, Category = "Settings")
+	FString GameView5;
+};
+
+UCLASS(config = EditorUserSettings)
+class GAMEPLAYDEBUGGER_API UGameplayDebuggerSettings : public UObject
+{
+	GENERATED_UCLASS_BODY()
+public:
+#if WITH_EDITOR
+	// Begin UObject Interface
+	virtual void PostEditChangeProperty(struct FPropertyChangedEvent& PropertyChangedEvent) override;
+	virtual void PostInitProperties() override;
+	// End UObject Interface
+#endif
+
+	uint32 GetSettings() 
+	{
+		uint32 Settings = 0;
+		Settings = OverHead ? Settings | (1 << EAIDebugDrawDataView::OverHead) : Settings & ~(1 << EAIDebugDrawDataView::OverHead);
+		Settings = Basic ? Settings | (1 << EAIDebugDrawDataView::Basic) : Settings & ~(1 << EAIDebugDrawDataView::Basic);
+		Settings = BehaviorTree ? Settings | (1 << EAIDebugDrawDataView::BehaviorTree) : Settings & ~(1 << EAIDebugDrawDataView::BehaviorTree);
+		Settings = EQS ? Settings | (1 << EAIDebugDrawDataView::EQS) : Settings & ~(1 << EAIDebugDrawDataView::EQS);
+		Settings = Perception ? Settings | (1 << EAIDebugDrawDataView::Perception) : Settings & ~(1 << EAIDebugDrawDataView::Perception);
+		Settings = GameView1 ? Settings | (1 << EAIDebugDrawDataView::GameView1) : Settings & ~(1 << EAIDebugDrawDataView::GameView1);
+		Settings = GameView2 ? Settings | (1 << EAIDebugDrawDataView::GameView2) : Settings & ~(1 << EAIDebugDrawDataView::GameView2);
+		Settings = GameView3 ? Settings | (1 << EAIDebugDrawDataView::GameView3) : Settings & ~(1 << EAIDebugDrawDataView::GameView3);
+		Settings = GameView4 ? Settings | (1 << EAIDebugDrawDataView::GameView4) : Settings & ~(1 << EAIDebugDrawDataView::GameView4);
+		Settings = GameView5 ? Settings | (1 << EAIDebugDrawDataView::GameView5) : Settings & ~(1 << EAIDebugDrawDataView::GameView5);
+
+		return Settings;
+	}
+
+	const FGDTCustomViewNames& GetCustomViewNames() { return CustomViewNames; }
+
+protected:
+	UPROPERTY(EditAnywhere, config, Category = "GameplayDebugger")
+	FGDTCustomViewNames CustomViewNames;
+	
+	UPROPERTY(EditAnywhere, config, Category = "GameplayDebugger")
+	bool OverHead;
+
+	UPROPERTY(EditAnywhere, config, Category = "GameplayDebugger")
+	bool Basic;
+
+	UPROPERTY(EditAnywhere, config, Category = "GameplayDebugger")
+	bool BehaviorTree;
+
+	UPROPERTY(EditAnywhere, config, Category = "GameplayDebugger|EQS")
+	bool EQS;
+
+	UPROPERTY(EditAnywhere, config, Category = "GameplayDebugger|EQS", meta = (EditCondition = "EQS"))
+	bool EnableEQSOnHUD;
+
+	UPROPERTY(EditAnywhere, config, Category = "GameplayDebugger|EQS", meta = (EditCondition = "EQS"))
+	int32 ActiveEQSIndex;
+
+	UPROPERTY(EditAnywhere, config, Category = "GameplayDebugger")
+	bool Perception;
+
+	UPROPERTY(EditAnywhere, config, Category = "GameplayDebugger")
+	bool GameView1;
+
+	UPROPERTY(EditAnywhere, config, Category = "GameplayDebugger")
+	bool GameView2;
+
+	UPROPERTY(EditAnywhere, config, Category = "GameplayDebugger")
+	bool GameView3;
+
+	UPROPERTY(EditAnywhere, config, Category = "GameplayDebugger")
+	bool GameView4;
+
+	UPROPERTY(EditAnywhere, config, Category = "GameplayDebugger")
+	bool GameView5;
+};
diff --git a/Engine/Source/Developer/GameplayDebugger/Classes/GameplayDebuggingHUDComponent.h b/Engine/Source/Developer/GameplayDebugger/Classes/GameplayDebuggingHUDComponent.h
--- a/Engine/Source/Developer/GameplayDebugger/Classes/GameplayDebuggingHUDComponent.h
+++ b/Engine/Source/Developer/GameplayDebugger/Classes/GameplayDebuggingHUDComponent.h
@@ -15,6 +15,11 @@ struct GAMEPLAYDEBUGGER_API FDebugCategoryView
 	FDebugCategoryView(EAIDebugDrawDataView::Type InView, const FString& Description) : Desc(Description), View(InView) {}
 };
 
+#define ADD_GAMEVIEW_CATEGORY(__Category__) \
+{\
+	UGameplayDebuggerSettings* GDS = UGameplayDebuggerSettings::StaticClass()->GetDefaultObject<UGameplayDebuggerSettings>();\
+	Categories.Add(FDebugCategoryView(EAIDebugDrawDataView::GameView1, GDS->GetCustomViewNames().GameView1.Len() ? GDS->GetCustomViewNames().GameView1 : TEXT("GameView1")));\
+}
 
 UCLASS(notplaceable)
 class GAMEPLAYDEBUGGER_API AGameplayDebuggingHUDComponent : public AActor
diff --git a/Engine/Source/Developer/GameplayDebugger/Classes/GameplayDebuggingTypes.h b/Engine/Source/Developer/GameplayDebugger/Classes/GameplayDebuggingTypes.h
--- a/Engine/Source/Developer/GameplayDebugger/Classes/GameplayDebuggingTypes.h
+++ b/Engine/Source/Developer/GameplayDebugger/Classes/GameplayDebuggingTypes.h
@@ -48,7 +48,6 @@ struct GAMEPLAYDEBUGGER_API FGameplayDebuggerSettings
 	}
 
 	uint32 &DebuggerShowFlags;
-	static uint32 ShowFlagIndex;
 };
 
 GAMEPLAYDEBUGGER_API
diff --git a/Engine/Source/Developer/GameplayDebugger/GameplayDebugger.Build.cs b/Engine/Source/Developer/GameplayDebugger/GameplayDebugger.Build.cs
--- a/Engine/Source/Developer/GameplayDebugger/GameplayDebugger.Build.cs
+++ b/Engine/Source/Developer/GameplayDebugger/GameplayDebugger.Build.cs
@@ -12,6 +12,7 @@ namespace UnrealBuildTool.Rules
                 new string[] {
 				    "Developer/GameplayDebugger/Public",
 				    "Developer/AIModule/Public",
+                    "Developer/Settings/Public",
 				    // ... add public include paths required here ...
 			    }
             );
@@ -35,6 +36,7 @@ namespace UnrealBuildTool.Rules
                     "RenderCore",
                     "RHI",
                     "ShaderCore",
+                    "Settings",
                     "AIModule",  // it have to be here for now. It'll be changed to remove any dependency to AIModule in future
 				}
                 );
@@ -49,6 +51,9 @@ namespace UnrealBuildTool.Rules
             if (UEBuildConfiguration.bBuildEditor == true)
             {
                 PrivateDependencyModuleNames.Add("UnrealEd");
+                PrivateDependencyModuleNames.Add("LevelEditor");
+                PrivateDependencyModuleNames.Add("Slate");
+                PublicIncludePaths.Add("Editor/LevelEditor/Public");
             }
 
             if (UEBuildConfiguration.bCompileRecast)
diff --git a/Engine/Source/Developer/GameplayDebugger/Private/GameplayDebugger.cpp b/Engine/Source/Developer/GameplayDebugger/Private/GameplayDebugger.cpp
--- a/Engine/Source/Developer/GameplayDebugger/Private/GameplayDebugger.cpp
+++ b/Engine/Source/Developer/GameplayDebugger/Private/GameplayDebugger.cpp
@@ -7,15 +7,19 @@
 #include "GameplayDebuggingReplicator.h"
 #include "GameplayDebuggingControllerComponent.h"
 #include "AISystem.h"
+#include "GameplayDebuggerSettings.h"
 #if WITH_EDITOR
 #include "Editor/EditorEngine.h"
+#include "ISettingsModule.h"
+#include "LevelEditor.h"
 #endif // WITH_EDITOR
 
-uint32 FGameplayDebuggerSettings::ShowFlagIndex = 0;
+#define ADD_LEVEL_EDITOR_EXTENSIONS 0
 
+#define LOCTEXT_NAMESPACE "FGameplayDebugger"
 FGameplayDebuggerSettings GameplayDebuggerSettings(class AGameplayDebuggingReplicator* Replicator)
 {
-	static uint32 Settings = (1 << EAIDebugDrawDataView::Basic) | (1 << EAIDebugDrawDataView::OverHead);
+	uint32 Settings = UGameplayDebuggerSettings::StaticClass()->GetDefaultObject<UGameplayDebuggerSettings>()->GetSettings();
 	return FGameplayDebuggerSettings(Replicator == NULL ? Settings : Replicator->DebuggerShowFlags);
 }
 
@@ -32,6 +36,10 @@ public:
 #if WITH_EDITOR
 	void OnLevelActorAdded(AActor* InActor);
 	void OnLevelActorDeleted(AActor* InActor);
+	TSharedRef<FExtender> OnExtendLevelEditorViewMenu(const TSharedRef<FUICommandList> CommandList);
+	void CreateSnappingOptionsMenu(FMenuBuilder& Builder);
+	void CreateSettingSubMenu(FMenuBuilder& Builder);
+
 #endif
 
 	TArray<TWeakObjectPtr<AGameplayDebuggingReplicator> >& GetAllReplicators(UWorld* InWorld);
@@ -48,6 +56,10 @@ private:
 	bool DoesGameplayDebuggingReplicatorExistForPlayerController(APlayerController* PlayerController);
 
 	TMap<TWeakObjectPtr<UWorld>, TArray<TWeakObjectPtr<AGameplayDebuggingReplicator> > > AllReplilcatorsPerWorlds;
+
+#if WITH_EDITOR
+	FLevelEditorModule::FLevelEditorMenuExtender ViewMenuExtender;
+#endif
 };
 
 IMPLEMENT_MODULE(FGameplayDebugger, GameplayDebugger)
@@ -63,6 +75,30 @@ void FGameplayDebugger::StartupModule()
 #if WITH_EDITOR
 		GEngine->OnLevelActorAdded().AddRaw(this, &FGameplayDebugger::OnLevelActorAdded);
 		GEngine->OnLevelActorDeleted().AddRaw(this, &FGameplayDebugger::OnLevelActorDeleted);
+
+		ISettingsModule* SettingsModule = FModuleManager::GetModulePtr<ISettingsModule>("Settings");
+		if (SettingsModule != nullptr)
+		{
+			SettingsModule->RegisterSettings("Editor", "General", "AITools",
+				LOCTEXT("AIToolsSettingsName", "AI Tools"),
+				LOCTEXT("AIToolsSettingsDescription", "General settings for UE4 AI Tools."),
+				UGameplayDebuggerSettings::StaticClass()->GetDefaultObject()
+				);
+		}
+#	if ADD_LEVEL_EDITOR_EXTENSIONS
+		if (!IsRunningCommandlet())
+		{
+			// Register the extension with the level editor
+			{
+				ViewMenuExtender = FLevelEditorModule::FLevelEditorMenuExtender::CreateRaw(this, &FGameplayDebugger::OnExtendLevelEditorViewMenu);
+				FLevelEditorModule* LevelEditor = FModuleManager::LoadModulePtr<FLevelEditorModule>(TEXT("LevelEditor"));
+				if (LevelEditor)
+				{
+					LevelEditor->GetAllLevelViewportOptionsMenuExtenders().Add(ViewMenuExtender);
+				}
+			}
+		}
+#	endif //ADD_LEVEL_EDITOR_EXTENSIONS
 #endif
 	}
 #endif
@@ -80,11 +116,78 @@ void FGameplayDebugger::ShutdownModule()
 #if WITH_EDITOR
 		GEngine->OnLevelActorAdded().RemoveAll(this);
 		GEngine->OnLevelActorDeleted().RemoveAll(this);
+
+		ISettingsModule* SettingsModule = FModuleManager::GetModulePtr<ISettingsModule>("Settings");
+		if (SettingsModule != nullptr)
+		{
+			SettingsModule->UnregisterSettings("Editor", "General", "AITools");
+		}
+
+#	if ADD_LEVEL_EDITOR_EXTENSIONS
+		if (UObjectInitialized() && !IsRunningCommandlet())
+		{
+			// Unregister the level editor extensions
+			{
+				FLevelEditorModule* LevelEditor = FModuleManager::LoadModulePtr<FLevelEditorModule>(TEXT("LevelEditor"));
+				if (LevelEditor)
+				{
+					LevelEditor->GetAllLevelViewportOptionsMenuExtenders().Remove(ViewMenuExtender);
+				}
+			}
+		}
+#	endif //ADD_LEVEL_EDITOR_EXTENSIONS
 #endif
 	}
 #endif
 }
 
+#if WITH_EDITOR
+void FGameplayDebugger::CreateSettingSubMenu(FMenuBuilder& Builder)
+{
+#if !(UE_BUILD_SHIPPING || UE_BUILD_TEST)
+	Builder.AddMenuEntry(
+		LOCTEXT("Test_GameplayDebugger_Menu", "Test Gameplay Debugger Option"),
+		LOCTEXT("Test_GameplayDebugger_Menu_Tooltip", "If Enabled, actors will snap to the nearest location on the constraint plane (NOTE: Only works correctly in perspective views right now!)"),
+		FSlateIcon(),
+		FUIAction(
+		FExecuteAction()/*FExecuteAction::CreateRaw(PlanarPolicy.Get(), &FPlanarConstraintSnapPolicy::ToggleEnabled)*/,
+		FCanExecuteAction(),
+		FIsActionChecked()/*FIsActionChecked::CreateRaw(PlanarPolicy.Get(), &FPlanarConstraintSnapPolicy::IsEnabled)*/
+		),
+		NAME_None,
+		EUserInterfaceActionType::Button,
+		NAME_None);
+#endif
+}
+
+void FGameplayDebugger::CreateSnappingOptionsMenu(FMenuBuilder& Builder)
+{
+#if !(UE_BUILD_SHIPPING || UE_BUILD_TEST)
+	Builder.AddMenuSeparator();
+	Builder.AddSubMenu(
+		LOCTEXT("Test_GameplayDebugger_Menu", "Gameplay Debugger"),
+		LOCTEXT("Test_GameplayDebugger_Menu_Tooltip", "Quick setting for Gameplay Debugger tool in selected view"),
+		FNewMenuDelegate::CreateRaw(this, &FGameplayDebugger::CreateSettingSubMenu)
+		);
+#endif
+}
+
+TSharedRef<FExtender> FGameplayDebugger::OnExtendLevelEditorViewMenu(const TSharedRef<FUICommandList> CommandList)
+{
+#if !(UE_BUILD_SHIPPING || UE_BUILD_TEST)
+	TSharedRef<FExtender> Extender(new FExtender());
+
+	Extender->AddMenuExtension(
+		"LevelViewportLayouts",
+		EExtensionHook::After,
+		NULL,
+		FMenuExtensionDelegate::CreateRaw(this, &FGameplayDebugger::CreateSnappingOptionsMenu));
+
+	return Extender;
+#endif
+}
+#endif
+
 bool FGameplayDebugger::DoesGameplayDebuggingReplicatorExistForPlayerController(APlayerController* PlayerController)
 {
 #if !(UE_BUILD_SHIPPING || UE_BUILD_TEST)
@@ -395,3 +498,4 @@ bool FGameplayDebugger::Exec(UWorld* Inworld, const TCHAR* Cmd, FOutputDevice& A
 	return bHandled;
 }
 
+#undef LOCTEXT_NAMESPACE
diff --git a/Engine/Source/Developer/GameplayDebugger/Private/GameplayDebuggerSettings.cpp b/Engine/Source/Developer/GameplayDebugger/Private/GameplayDebuggerSettings.cpp
new file mode 100644
--- /dev/null
+++ b/Engine/Source/Developer/GameplayDebugger/Private/GameplayDebuggerSettings.cpp
@@ -0,0 +1,70 @@
+// Copyright 1998-2014 Epic Games, Inc. All Rights Reserved.
+
+#include "GameplayDebuggerPrivate.h"
+#include "Misc/CoreMisc.h"
+#include "GameplayDebuggingComponent.h"
+#include "GameplayDebuggingHUDComponent.h"
+#include "GameplayDebuggingReplicator.h"
+#include "GameplayDebuggingControllerComponent.h"
+#include "AISystem.h"
+#include "GameplayDebuggerSettings.h"
+#if WITH_EDITOR
+#include "Editor/EditorEngine.h"
+#include "ISettingsModule.h"
+#endif // WITH_EDITOR
+
+UGameplayDebuggerSettings::UGameplayDebuggerSettings(const FObjectInitializer& ObjectInitializer) : Super(ObjectInitializer)
+{
+}
+
+#if WITH_EDITOR
+void UGameplayDebuggerSettings::PostInitProperties()
+{
+	Super::PostInitProperties();
+
+#define UPDATE_GAMEVIEW_DISPLAYNAME(__GameView__) \
+	{\
+		UProperty* Prop = FindField<UProperty>(UGameplayDebuggerSettings::StaticClass(), TEXT(#__GameView__));\
+		if (Prop)\
+		{\
+			Prop->SetPropertyFlags(CPF_Edit);\
+			Prop->SetMetaData(TEXT("DisplayName"), CustomViewNames.__GameView__.Len() > 0 ? *CustomViewNames.__GameView__ : TEXT(#__GameView__));\
+		}\
+	}
+	UPDATE_GAMEVIEW_DISPLAYNAME(GameView1);
+	UPDATE_GAMEVIEW_DISPLAYNAME(GameView2);
+	UPDATE_GAMEVIEW_DISPLAYNAME(GameView3);
+	UPDATE_GAMEVIEW_DISPLAYNAME(GameView4);
+	UPDATE_GAMEVIEW_DISPLAYNAME(GameView5);
+#undef UPDATE_GAMEVIEW_DISPLAYNAME
+}
+
+void UGameplayDebuggerSettings::PostEditChangeProperty(FPropertyChangedEvent& PropertyChangedEvent)
+{
+	Super::PostEditChangeProperty(PropertyChangedEvent);
+
+	const FName Name = (PropertyChangedEvent.Property != nullptr) ? PropertyChangedEvent.Property->GetFName() : NAME_None;
+	const FName MemberPropertyName = (PropertyChangedEvent.MemberProperty != nullptr) ? PropertyChangedEvent.MemberProperty->GetFName() : NAME_None;
+
+#define UPDATE_GAMEVIEW_DISPLAYNAME(__GameView__) \
+	if (MemberPropertyName == (TEXT("CustomViewNames"))) \
+	{\
+		if (Name == FName(TEXT(#__GameView__)) ) \
+		{ \
+			UProperty* Prop = FindField<UProperty>(UGameplayDebuggerSettings::StaticClass(), *Name.ToString());\
+			check(Prop);\
+			Prop->SetPropertyFlags(CPF_Edit); \
+			Prop->SetMetaData(TEXT("DisplayName"), CustomViewNames.__GameView__.Len() > 0 ? *CustomViewNames.__GameView__ : *Name.ToString());\
+		}\
+	}
+
+	UPDATE_GAMEVIEW_DISPLAYNAME(GameView1);
+	UPDATE_GAMEVIEW_DISPLAYNAME(GameView2);
+	UPDATE_GAMEVIEW_DISPLAYNAME(GameView3);
+	UPDATE_GAMEVIEW_DISPLAYNAME(GameView4);
+	UPDATE_GAMEVIEW_DISPLAYNAME(GameView5);
+#undef UPDATE_GAMEVIEW_DISPLAYNAME
+
+	GEditor->SaveEditorUserSettings();
+}
+#endif
diff --git a/Engine/Source/Developer/GameplayDebugger/Private/GameplayDebuggingHUDComponent.cpp b/Engine/Source/Developer/GameplayDebugger/Private/GameplayDebuggingHUDComponent.cpp
--- a/Engine/Source/Developer/GameplayDebugger/Private/GameplayDebuggingHUDComponent.cpp
+++ b/Engine/Source/Developer/GameplayDebugger/Private/GameplayDebuggingHUDComponent.cpp
@@ -4,6 +4,7 @@
 	HUD.cpp: Heads up Display related functionality
 =============================================================================*/
 #include "GameplayDebuggerPrivate.h"
+#include "GameplayDebuggerSettings.h"
 #include "Net/UnrealNetwork.h"
 #include "GameplayDebuggingComponent.h"
 #include "GameplayDebuggingHUDComponent.h"
@@ -73,11 +74,11 @@ AGameplayDebuggingReplicator* AGameplayDebuggingHUDComponent::GetDebuggingReplic
 
 void AGameplayDebuggingHUDComponent::GetKeyboardDesc(TArray<FDebugCategoryView>& Categories)
 {
-	Categories.Add(FDebugCategoryView(EAIDebugDrawDataView::NavMesh, "NavMesh"));			// Num0
-	Categories.Add(FDebugCategoryView(EAIDebugDrawDataView::Basic, "Basic"));
-	Categories.Add(FDebugCategoryView(EAIDebugDrawDataView::BehaviorTree, "Behavior"));
-	Categories.Add(FDebugCategoryView(EAIDebugDrawDataView::EQS, "EQS"));
-	Categories.Add(FDebugCategoryView(EAIDebugDrawDataView::Perception, "Perception"));
+	Categories.Add(FDebugCategoryView(EAIDebugDrawDataView::NavMesh, TEXT("NavMesh")));
+	Categories.Add(FDebugCategoryView(EAIDebugDrawDataView::Basic, TEXT("Basic")));
+	Categories.Add(FDebugCategoryView(EAIDebugDrawDataView::BehaviorTree, TEXT("BehaviorTree")));
+	Categories.Add(FDebugCategoryView(EAIDebugDrawDataView::EQS, TEXT("EQS")));
+	Categories.Add(FDebugCategoryView(EAIDebugDrawDataView::Perception, TEXT("Perception")));
 }
 
 void AGameplayDebuggingHUDComponent::PrintAllData()
diff --git a/Engine/Source/Editor/MainFrame/Private/Menus/MainMenu.cpp b/Engine/Source/Editor/MainFrame/Private/Menus/MainMenu.cpp
--- a/Engine/Source/Editor/MainFrame/Private/Menus/MainMenu.cpp
+++ b/Engine/Source/Editor/MainFrame/Private/Menus/MainMenu.cpp
@@ -176,7 +176,7 @@ void FMainMenu::FillWindowMenu( FMenuBuilder& MenuBuilder, const TSharedRef< FEx
 		bool bVisualLogger = GetDefault<UEditorExperimentalSettings>()->bVisualLogger;
 
 		// Make sure at least one is enabled before creating the section
-		if (bProjectLauncher || bMessagingDebugger || bBlutility || bTranslationEditor)
+		if (bProjectLauncher || bMessagingDebugger || bBlutility || bTranslationEditor || bVisualLogger)
 		{
 			MenuBuilder.BeginSection("ExperimentalTabSpawners", LOCTEXT("ExperimentalTabSpawnersHeading", "Experimental"));
 			{
