Added default engine settings for GameplayDebugger. User can override that settings in editor.

diff --git a/Engine/Config/BaseEngine.ini b/Engine/Config/BaseEngine.ini
--- a/Engine/Config/BaseEngine.ini
+++ b/Engine/Config/BaseEngine.ini
@@ -2402,3 +2402,23 @@ NumberOfHistoryItemsTracked=5
 [VisualLogger]
 FrameCacheLenght=1.0f ;in seconds, to batch log data between file serializations
 UseCompression=false ;works only with binary files
+
+[GameplayDebuggerSettings]
+OverHead=True
+Basic=True
+BehaviorTree=False
+EQS=False
+EnableEQSOnHUD=true
+Perception=False
+GameView1=False
+GameView2=False
+GameView3=False
+GameView4=False
+GameView5=False
+NameForGameView1="GameView1"
+NameForGameView2="GameView2"
+NameForGameView3="GameView3"
+NameForGameView4="GameView4"
+NameForGameView5="GameView5"
+
+
diff --git a/Engine/Source/Developer/GameplayDebugger/Private/GameplayDebuggerSettings.cpp b/Engine/Source/Developer/GameplayDebugger/Private/GameplayDebuggerSettings.cpp
--- a/Engine/Source/Developer/GameplayDebugger/Private/GameplayDebuggerSettings.cpp
+++ b/Engine/Source/Developer/GameplayDebugger/Private/GameplayDebuggerSettings.cpp
@@ -15,6 +15,28 @@
 
 UGameplayDebuggerSettings::UGameplayDebuggerSettings(const FObjectInitializer& ObjectInitializer) : Super(ObjectInitializer)
 {
+	if (HasAnyFlags(RF_ClassDefaultObject))
+	{
+		/*read base default values from Engine config file. It can be overridden (locally) by Editor settings*/
+		GConfig->GetString(TEXT("GameplayDebuggerSettings"), TEXT("NameForGameView1"), CustomViewNames.GameView1, GEngineIni);
+		GConfig->GetString(TEXT("GameplayDebuggerSettings"), TEXT("NameForGameView2"), CustomViewNames.GameView2, GEngineIni);
+		GConfig->GetString(TEXT("GameplayDebuggerSettings"), TEXT("NameForGameView3"), CustomViewNames.GameView3, GEngineIni);
+		GConfig->GetString(TEXT("GameplayDebuggerSettings"), TEXT("NameForGameView4"), CustomViewNames.GameView4, GEngineIni);
+		GConfig->GetString(TEXT("GameplayDebuggerSettings"), TEXT("NameForGameView5"), CustomViewNames.GameView5, GEngineIni);
+		GConfig->GetString(TEXT("GameplayDebuggerSettings"), TEXT("NameForGameView5"), CustomViewNames.GameView5, GEngineIni);
+		
+		GConfig->GetBool(TEXT("GameplayDebuggerSettings"), TEXT("OverHead"), OverHead, GEngineIni);
+		GConfig->GetBool(TEXT("GameplayDebuggerSettings"), TEXT("Basic"), Basic, GEngineIni);
+		GConfig->GetBool(TEXT("GameplayDebuggerSettings"), TEXT("BehaviorTree"), BehaviorTree, GEngineIni);
+		GConfig->GetBool(TEXT("GameplayDebuggerSettings"), TEXT("EQS"), EQS, GEngineIni);
+		GConfig->GetBool(TEXT("GameplayDebuggerSettings"), TEXT("EnableEQSOnHUD"), EnableEQSOnHUD, GEngineIni);
+		GConfig->GetBool(TEXT("GameplayDebuggerSettings"), TEXT("Perception"), Perception, GEngineIni);
+		GConfig->GetBool(TEXT("GameplayDebuggerSettings"), TEXT("GameView1"), GameView1, GEngineIni);
+		GConfig->GetBool(TEXT("GameplayDebuggerSettings"), TEXT("GameView2"), GameView2, GEngineIni);
+		GConfig->GetBool(TEXT("GameplayDebuggerSettings"), TEXT("GameView3"), GameView3, GEngineIni);
+		GConfig->GetBool(TEXT("GameplayDebuggerSettings"), TEXT("GameView4"), GameView4, GEngineIni);
+		GConfig->GetBool(TEXT("GameplayDebuggerSettings"), TEXT("GameView5"), GameView5, GEngineIni);
+	}
 }
 
 #if WITH_EDITOR
@@ -23,12 +45,12 @@ void UGameplayDebuggerSettings::PostInitProperties()
 	Super::PostInitProperties();
 
 #define UPDATE_GAMEVIEW_DISPLAYNAME(__GameView__) \
-	{\
+			{\
 		UProperty* Prop = FindField<UProperty>(UGameplayDebuggerSettings::StaticClass(), TEXT(#__GameView__));\
 		if (Prop)\
-		{\
+				{\
 			Prop->SetPropertyFlags(CPF_Edit);\
-			Prop->SetMetaData(TEXT("DisplayName"), CustomViewNames.__GameView__.Len() > 0 ? *CustomViewNames.__GameView__ : TEXT(#__GameView__));\
+			Prop->SetMetaData(TEXT("DisplayName"), CustomViewNames.__GameView__.Len() > 0 ? *CustomViewNames.__GameView__ : TEXT(#__GameView__)); \
 		}\
 	}
 	UPDATE_GAMEVIEW_DISPLAYNAME(GameView1);
@@ -54,7 +76,7 @@ void UGameplayDebuggerSettings::PostEditChangeProperty(FPropertyChangedEvent& Pr
 			UProperty* Prop = FindField<UProperty>(UGameplayDebuggerSettings::StaticClass(), *Name.ToString());\
 			check(Prop);\
 			Prop->SetPropertyFlags(CPF_Edit); \
-			Prop->SetMetaData(TEXT("DisplayName"), CustomViewNames.__GameView__.Len() > 0 ? *CustomViewNames.__GameView__ : *Name.ToString());\
+			Prop->SetMetaData(TEXT("DisplayName"), CustomViewNames.__GameView__.Len() > 0 ? *CustomViewNames.__GameView__ : TEXT(#__GameView__)); \
 		}\
 	}
 
