Updates to settings for LogVisualizer and GameplayDebugger. Added GameplayDebugger to experimental settings.

diff --git a/Engine/Source/Developer/GameplayDebugger/Private/GameplayDebugger.cpp b/Engine/Source/Developer/GameplayDebugger/Private/GameplayDebugger.cpp
--- a/Engine/Source/Developer/GameplayDebugger/Private/GameplayDebugger.cpp
+++ b/Engine/Source/Developer/GameplayDebugger/Private/GameplayDebugger.cpp
@@ -39,7 +39,7 @@ public:
 	TSharedRef<FExtender> OnExtendLevelEditorViewMenu(const TSharedRef<FUICommandList> CommandList);
 	void CreateSnappingOptionsMenu(FMenuBuilder& Builder);
 	void CreateSettingSubMenu(FMenuBuilder& Builder);
-
+	void HandleExperimentalSettingChanged(FName PropertyName);
 #endif
 
 	TArray<TWeakObjectPtr<AGameplayDebuggingReplicator> >& GetAllReplicators(UWorld* InWorld);
@@ -76,11 +76,13 @@ void FGameplayDebugger::StartupModule()
 		GEngine->OnLevelActorAdded().AddRaw(this, &FGameplayDebugger::OnLevelActorAdded);
 		GEngine->OnLevelActorDeleted().AddRaw(this, &FGameplayDebugger::OnLevelActorDeleted);
 
+		UEditorExperimentalSettings::StaticClass()->GetDefaultObject<UEditorExperimentalSettings>()->OnSettingChanged().AddRaw(this, &FGameplayDebugger::HandleExperimentalSettingChanged);
+
 		ISettingsModule* SettingsModule = FModuleManager::GetModulePtr<ISettingsModule>("Settings");
 		if (SettingsModule != nullptr)
 		{
-			SettingsModule->RegisterSettings("Editor", "General", "AITools",
-				LOCTEXT("AIToolsSettingsName", "AI Tools"),
+			SettingsModule->RegisterSettings("Editor", "General", "GameplayDebugger",
+				LOCTEXT("AIToolsSettingsName", "Gameplay Debugger"),
 				LOCTEXT("AIToolsSettingsDescription", "General settings for UE4 AI Tools."),
 				UGameplayDebuggerSettings::StaticClass()->GetDefaultObject()
 				);
@@ -117,10 +119,12 @@ void FGameplayDebugger::ShutdownModule()
 		GEngine->OnLevelActorAdded().RemoveAll(this);
 		GEngine->OnLevelActorDeleted().RemoveAll(this);
 
+		UEditorExperimentalSettings::StaticClass()->GetDefaultObject<UEditorExperimentalSettings>()->OnSettingChanged().RemoveAll(this);
+
 		ISettingsModule* SettingsModule = FModuleManager::GetModulePtr<ISettingsModule>("Settings");
 		if (SettingsModule != nullptr)
 		{
-			SettingsModule->UnregisterSettings("Editor", "General", "AITools");
+			SettingsModule->UnregisterSettings("Editor", "General", "GameplayDebugger");
 		}
 
 #	if ADD_LEVEL_EDITOR_EXTENSIONS
@@ -142,6 +146,37 @@ void FGameplayDebugger::ShutdownModule()
 }
 
 #if WITH_EDITOR
+
+void FGameplayDebugger::HandleExperimentalSettingChanged(FName PropertyName)
+{
+	if (PropertyName == TEXT("bGameplayDebugger"))
+	{
+		if (UEditorExperimentalSettings::StaticClass()->GetDefaultObject<UEditorExperimentalSettings>()->bGameplayDebugger == false)
+		{
+			ISettingsModule* SettingsModule = FModuleManager::GetModulePtr<ISettingsModule>("Settings");
+			if (SettingsModule != nullptr)
+			{
+				SettingsModule->UnregisterSettings("Editor", "General", "GameplayDebugger");
+			}
+		}
+		else
+		{
+			ISettingsModule* SettingsModule = FModuleManager::GetModulePtr<ISettingsModule>("Settings");
+			if (SettingsModule != nullptr)
+			{
+				SettingsModule->RegisterSettings("Editor", "General", "GameplayDebugger",
+					LOCTEXT("AIToolsSettingsName", "Gameplay Debugger"),
+					LOCTEXT("AIToolsSettingsDescription", "General settings for UE4 AI Tools."),
+					UGameplayDebuggerSettings::StaticClass()->GetDefaultObject()
+					);
+			}
+		}
+	}
+	//FLevelEditorModule& LevelEditorModule = FModuleManager::GetModuleChecked<FLevelEditorModule>("LevelEditor");
+	//TSharedPtr<FTabManager> LevelEditorTabManager = LevelEditorModule.GetLevelEditorTabManager();
+	//LevelEditorTabManager->UpdateMainMenu(true);
+}
+
 void FGameplayDebugger::CreateSettingSubMenu(FMenuBuilder& Builder)
 {
 #if !(UE_BUILD_SHIPPING || UE_BUILD_TEST)
@@ -163,12 +198,15 @@ void FGameplayDebugger::CreateSettingSubMenu(FMenuBuilder& Builder)
 void FGameplayDebugger::CreateSnappingOptionsMenu(FMenuBuilder& Builder)
 {
 #if !(UE_BUILD_SHIPPING || UE_BUILD_TEST)
-	Builder.AddMenuSeparator();
-	Builder.AddSubMenu(
-		LOCTEXT("Test_GameplayDebugger_Menu", "Gameplay Debugger"),
-		LOCTEXT("Test_GameplayDebugger_Menu_Tooltip", "Quick setting for Gameplay Debugger tool in selected view"),
-		FNewMenuDelegate::CreateRaw(this, &FGameplayDebugger::CreateSettingSubMenu)
-		);
+	if (GCurrentLevelEditingViewportClient && GCurrentLevelEditingViewportClient->EngineShowFlags.DebugAI && GCurrentLevelEditingViewportClient->IsSimulateInEditorViewport())
+	{
+		Builder.AddMenuSeparator();
+		Builder.AddSubMenu(
+			LOCTEXT("Test_GameplayDebugger_Menu", "Gameplay Debugger"),
+			LOCTEXT("Test_GameplayDebugger_Menu_Tooltip", "Quick setting for Gameplay Debugger tool in selected view"),
+			FNewMenuDelegate::CreateRaw(this, &FGameplayDebugger::CreateSettingSubMenu)
+			);
+	}
 #endif
 }
 
diff --git a/Engine/Source/Developer/GameplayDebugger/Private/GameplayDebuggingReplicator.cpp b/Engine/Source/Developer/GameplayDebugger/Private/GameplayDebuggingReplicator.cpp
--- a/Engine/Source/Developer/GameplayDebugger/Private/GameplayDebuggingReplicator.cpp
+++ b/Engine/Source/Developer/GameplayDebugger/Private/GameplayDebuggingReplicator.cpp
@@ -463,6 +463,11 @@ void AGameplayDebuggingReplicator::OnDebugAIDelegate(class UCanvas* Canvas, clas
 	{
 		return;
 	}
+
+	if (!Canvas || !Canvas->SceneView || Canvas->SceneView->bIsGameView == false)
+	{
+		return;
+	}
 	LastDrawAtFrame = GFrameNumber;
 
 	FEngineShowFlags EngineShowFlags = Canvas && Canvas->SceneView && Canvas->SceneView->Family ? Canvas->SceneView->Family->EngineShowFlags : FEngineShowFlags(GIsEditor ? EShowFlagInitMode::ESFIM_Editor : EShowFlagInitMode::ESFIM_Game);
diff --git a/Engine/Source/Developer/NewLogVisualizer/Private/LogVisualizerModule.cpp b/Engine/Source/Developer/NewLogVisualizer/Private/LogVisualizerModule.cpp
--- a/Engine/Source/Developer/NewLogVisualizer/Private/LogVisualizerModule.cpp
+++ b/Engine/Source/Developer/NewLogVisualizer/Private/LogVisualizerModule.cpp
@@ -5,7 +5,12 @@
 #include "LogVisualizerStyle.h"
 #include "SDockTab.h"
 #include "VisualLoggerRenderingActor.h"
+#include "LogVIsualizerSettings.h"
+#if WITH_EDITOR
 #include "Editor/EditorEngine.h"
+#include "ISettingsModule.h"
+#include "LevelEditor.h"
+#endif // WITH_EDITOR
 
 #define LOCTEXT_NAMESPACE "FLogVisualizerModule"
 
@@ -23,6 +28,18 @@ void FNewLogVisualizerModule::StartupModule()
 	// This is still experimental in the editor, so it'll be invoked specifically in FMainMenu if the experimental settings flag is set.
 	// When no longer experimental, switch to the nomad spawner registration below
 	FGlobalTabmanager::Get()->RegisterTabSpawner(VisualLoggerTabName, FOnSpawnTab::CreateRaw(this, &FNewLogVisualizerModule::SpawnLogVisualizerTab));
+
+	UEditorExperimentalSettings::StaticClass()->GetDefaultObject<UEditorExperimentalSettings>()->OnSettingChanged().AddRaw(this, &FNewLogVisualizerModule::HandleExperimentalSettingChanged);
+	
+	ISettingsModule* SettingsModule = FModuleManager::GetModulePtr<ISettingsModule>("Settings");
+	if (SettingsModule != nullptr)
+	{
+		SettingsModule->RegisterSettings("Editor", "General", "VisualLogger",
+			LOCTEXT("AIToolsSettingsName", "Visual Logger"),
+			LOCTEXT("AIToolsSettingsDescription", "General settings for UE4 AI Tools."),
+			ULogVisualizerSettings::StaticClass()->GetDefaultObject()
+			);
+	}
 }
 
 void FNewLogVisualizerModule::ShutdownModule()
@@ -30,10 +47,48 @@ void FNewLogVisualizerModule::ShutdownModule()
 	FGlobalTabmanager::Get()->UnregisterTabSpawner(VisualLoggerTabName);
 	FVisualLoggerCommands::Unregister();
 	IModularFeatures::Get().UnregisterModularFeature(VisualLoggerTabName, this);
+	
+	//UEditorExperimentalSettings::StaticClass()->GetDefaultObject<UEditorExperimentalSettings>()->OnSettingChanged().RemoveAll(this);
+
+	ISettingsModule* SettingsModule = FModuleManager::GetModulePtr<ISettingsModule>("Settings");
+	if (SettingsModule != nullptr)
+	{
+		SettingsModule->UnregisterSettings("Editor", "General", "VisualLogger");
+	}
 
 	FLogVisualizerStyle::Shutdown();
 }
 
+#if WITH_EDITOR
+
+void FNewLogVisualizerModule::HandleExperimentalSettingChanged(FName PropertyName)
+{
+	if (PropertyName == TEXT("bVisualLogger"))
+	{
+		if (UEditorExperimentalSettings::StaticClass()->GetDefaultObject<UEditorExperimentalSettings>()->bVisualLogger == false)
+		{
+			ISettingsModule* SettingsModule = FModuleManager::GetModulePtr<ISettingsModule>("Settings");
+			if (SettingsModule != nullptr)
+			{
+				SettingsModule->UnregisterSettings("Editor", "General", "VisualLogger");
+			}
+		}
+		else
+		{
+			ISettingsModule* SettingsModule = FModuleManager::GetModulePtr<ISettingsModule>("Settings");
+			if (SettingsModule != nullptr)
+			{
+				SettingsModule->RegisterSettings("Editor", "General", "VisualLogger",
+					LOCTEXT("AIToolsSettingsName", "Visual Logger"),
+					LOCTEXT("AIToolsSettingsDescription", "General settings for UE4 AI Tools."),
+					ULogVisualizerSettings::StaticClass()->GetDefaultObject()
+					);
+			}
+		}
+	}
+}
+#endif
+
 TSharedRef<SDockTab> FNewLogVisualizerModule::SpawnLogVisualizerTab(const FSpawnTabArgs& SpawnTabArgs)
 {
 	const TSharedRef<SDockTab> MajorTab = SNew(SDockTab)
diff --git a/Engine/Source/Developer/NewLogVisualizer/Private/LogVisualizerSettings.cpp b/Engine/Source/Developer/NewLogVisualizer/Private/LogVisualizerSettings.cpp
new file mode 100644
--- /dev/null
+++ b/Engine/Source/Developer/NewLogVisualizer/Private/LogVisualizerSettings.cpp
@@ -0,0 +1,28 @@
+// Copyright 1998-2014 Epic Games, Inc. All Rights Reserved.
+
+#include "LogVisualizer.h"
+#include "Misc/CoreMisc.h"
+#include "LogVIsualizerSettings.h"
+#if WITH_EDITOR
+#include "Editor/EditorEngine.h"
+#include "ISettingsModule.h"
+#endif // WITH_EDITOR
+
+ULogVisualizerSettings::ULogVisualizerSettings(const FObjectInitializer& ObjectInitializer) : Super(ObjectInitializer)
+{
+	TrivialLogsThreshold = 1;
+}
+
+void ULogVisualizerSettings::PostEditChangeProperty(struct FPropertyChangedEvent& PropertyChangedEvent)
+{
+	Super::PostEditChangeProperty(PropertyChangedEvent);
+
+	const FName Name = (PropertyChangedEvent.Property != nullptr) ? PropertyChangedEvent.Property->GetFName() : NAME_None;
+	if (!FUnrealEdMisc::Get().IsDeletePreferences())
+	{
+		SaveConfig();
+	}
+
+	SettingChangedEvent.Broadcast(Name);
+}
+
diff --git a/Engine/Source/Developer/NewLogVisualizer/Private/STimeline.cpp b/Engine/Source/Developer/NewLogVisualizer/Private/STimeline.cpp
--- a/Engine/Source/Developer/NewLogVisualizer/Private/STimeline.cpp
+++ b/Engine/Source/Developer/NewLogVisualizer/Private/STimeline.cpp
@@ -3,6 +3,12 @@
 #include "STimelineBar.h"
 #include "TimeSliderController.h"
 #include "STimelinesContainer.h"
+#include "LogVisualizerSettings.h"
+
+STimeline::~STimeline()
+{
+	ULogVisualizerSettings::StaticClass()->GetDefaultObject<ULogVisualizerSettings>()->OnSettingChanged().RemoveAll(this);
+}
 
 FReply STimeline::OnMouseButtonDown(const FGeometry& MyGeometry, const FPointerEvent& MouseEvent)
 {
@@ -39,7 +45,9 @@ void STimeline::OnDeselect()
 
 void STimeline::UpdateVisibility()
 {
-	SetVisibility(HiddenEntries.Num() == Entries.Num() || (SearchFilter.Len() > 0 && Name.ToString().Find(SearchFilter) == INDEX_NONE) ? EVisibility::Collapsed : EVisibility::Visible);
+	ULogVisualizerSettings* Settings = ULogVisualizerSettings::StaticClass()->GetDefaultObject<ULogVisualizerSettings>();
+	const bool bJustIgnore = Settings->bIgnoreTrivialLogs && Entries.Num() <= Settings->TrivialLogsThreshold;
+	SetVisibility(bJustIgnore || HiddenEntries.Num() == Entries.Num() || (SearchFilter.Len() > 0 && Name.ToString().Find(SearchFilter) == INDEX_NONE) ? EVisibility::Collapsed : EVisibility::Visible);
 	if (GetVisibility() != EVisibility::Visible)
 	{
 		Owner->SetSelectionState(SharedThis(this), false, false);
@@ -74,6 +82,11 @@ bool STimeline::IsEntryHidden(const FVisualLogDevice::FVisualLogEntryItem& Entry
 	return HiddenEntries.Contains(&EntryItem);
 }
 
+void STimeline::HandleLogVisualizerSettingChanged(FName Name)
+{
+	UpdateVisibility();
+}
+
 void STimeline::Construct(const FArguments& InArgs, TSharedPtr<SVisualLoggerView> VisualLoggerView, TSharedPtr<FSequencerTimeSliderController> TimeSliderController, TSharedPtr<STimelinesContainer> InContainer, const FVisualLogDevice::FVisualLogEntryItem& Entry)
 {
 	VisualLoggerInterface = InArgs._VisualLoggerInterface.Get();
@@ -82,6 +95,9 @@ void STimeline::Construct(const FArguments& InArgs, TSharedPtr<SVisualLoggerView
 
 	Entries.Add(Entry);
 	OnFiltersChanged();
+	
+	ULogVisualizerSettings::StaticClass()->GetDefaultObject<ULogVisualizerSettings>()->OnSettingChanged().AddRaw(this, &STimeline::HandleLogVisualizerSettingChanged);
+
 
 	ChildSlot
 		[
diff --git a/Engine/Source/Developer/NewLogVisualizer/Private/STimeline.h b/Engine/Source/Developer/NewLogVisualizer/Private/STimeline.h
--- a/Engine/Source/Developer/NewLogVisualizer/Private/STimeline.h
+++ b/Engine/Source/Developer/NewLogVisualizer/Private/STimeline.h
@@ -20,6 +20,7 @@ public:
 	virtual void OnDeselect();
 
 	/** Constructs this widget with InArgs */
+	virtual ~STimeline();
 	void Construct(const FArguments& InArgs, TSharedPtr<class SVisualLoggerView> VisualLoggerView, TSharedPtr<class FSequencerTimeSliderController> InTimeSliderController, TSharedPtr<class STimelinesContainer> InContainer, const FVisualLogDevice::FVisualLogEntryItem& Entry);
 	bool IsSelected() const;
 	bool IsEntryHidden(const FVisualLogDevice::FVisualLogEntryItem&) const;
@@ -31,6 +32,7 @@ public:
 	void AddEntry(const FVisualLogDevice::FVisualLogEntryItem& Entry) { Entries.Add(Entry); }
 	const TArray<FVisualLogDevice::FVisualLogEntryItem>& GetEntries() { return Entries; }
 	FName GetName() { return Name; }
+	void HandleLogVisualizerSettingChanged(FName Name);
 
 protected:
 	TArray<FVisualLogDevice::FVisualLogEntryItem> Entries;
diff --git a/Engine/Source/Developer/NewLogVisualizer/Private/SVisualLogger.cpp b/Engine/Source/Developer/NewLogVisualizer/Private/SVisualLogger.cpp
--- a/Engine/Source/Developer/NewLogVisualizer/Private/SVisualLogger.cpp
+++ b/Engine/Source/Developer/NewLogVisualizer/Private/SVisualLogger.cpp
@@ -272,6 +272,23 @@ void SVisualLogger::Construct(const FArguments& InArgs, const TSharedRef<SDockTa
 		"Window"
 		);
 
+	MenuBarBuilder.AddMenuEntry(
+		LOCTEXT("SettingsMenuLabel", "Settings"),
+		FText::GetEmpty(),
+		FSlateIcon(),
+		FUIAction(
+			FExecuteAction::CreateLambda(
+			[this](){
+				ISettingsModule* SettingsModule = FModuleManager::GetModulePtr<ISettingsModule>("Settings");
+				if (SettingsModule != nullptr)
+				{
+					SettingsModule->ShowViewer("Editor", "General", "VisualLogger");
+				}
+			}
+		)),
+		"Settings"
+		);
+
 
 	ChildSlot
 		[
diff --git a/Engine/Source/Developer/NewLogVisualizer/Public/LogVisualizerModule.h b/Engine/Source/Developer/NewLogVisualizer/Public/LogVisualizerModule.h
--- a/Engine/Source/Developer/NewLogVisualizer/Public/LogVisualizerModule.h
+++ b/Engine/Source/Developer/NewLogVisualizer/Public/LogVisualizerModule.h
@@ -17,6 +17,7 @@ public:
 	// End IModuleInterface
 
 private:
+	void HandleExperimentalSettingChanged(FName PropertyName);
 	TSharedRef<SDockTab> SpawnLogVisualizerTab(const FSpawnTabArgs& SpawnTabArgs);
 	void OnTabClosed(TSharedRef<SDockTab>);
 };
diff --git a/Engine/Source/Developer/NewLogVisualizer/Public/LogVisualizerSettings.h b/Engine/Source/Developer/NewLogVisualizer/Public/LogVisualizerSettings.h
new file mode 100644
--- /dev/null
+++ b/Engine/Source/Developer/NewLogVisualizer/Public/LogVisualizerSettings.h
@@ -0,0 +1,54 @@
+// Copyright 1998-2014 Epic Games, Inc. All Rights Reserved.
+
+/*=============================================================================
+GameplayDebuggerSettings.h: Declares the UGameplayDebuggerSettings class.
+=============================================================================*/
+#pragma once
+
+
+#include "LogVisualizerSettings.generated.h"
+
+UCLASS(config = EditorUserSettings)
+class NEWLOGVISUALIZER_API ULogVisualizerSettings : public UObject
+{
+	GENERATED_UCLASS_BODY()
+public:
+#if WITH_EDITOR
+	// Begin UObject Interface
+	//virtual void PostEditChangeProperty(struct FPropertyChangedEvent& PropertyChangedEvent) override;
+	//virtual void PostInitProperties() override;
+	// End UObject Interface
+#endif
+
+public:
+	/**Whether to show trivial logs, i.e. the ones with only one entry.*/
+	UPROPERTY(EditAnywhere, config, Category = "VisualLogger")
+	bool bIgnoreTrivialLogs;
+
+	/**Threshold for trivial Logs*/
+	UPROPERTY(EditAnywhere, config, Category = "VisualLogger", meta = (EditCondition = "bIgnoreTrivialLogs", ClampMin = "0", ClampMax = "10", UIMin = "0", UIMax = "10"))
+	int32 TrivialLogsThreshold;
+
+	/**Whether to show the recent data or not.*/
+	UPROPERTY(EditAnywhere, config, Category = "VisualLogger")
+	bool bStickToRecentData;
+
+	/**Whether to show histogram labels inside graph or outside.*/
+	UPROPERTY(EditAnywhere, config, Category = "VisualLogger")
+	bool bShowHistogramLabelsOutside;
+
+	DECLARE_EVENT_OneParam(ULogVisualizerSettings, FSettingChangedEvent, FName /*PropertyName*/);
+	FSettingChangedEvent& OnSettingChanged() { return SettingChangedEvent; }
+
+protected:
+
+	// UObject overrides
+
+	virtual void PostEditChangeProperty(struct FPropertyChangedEvent& PropertyChangedEvent) override;
+
+private:
+
+	// Holds an event delegate that is executed when a setting has changed.
+	FSettingChangedEvent SettingChangedEvent;
+
+};
diff --git a/Engine/Source/Editor/UnrealEd/Classes/Settings/EditorExperimentalSettings.h b/Engine/Source/Editor/UnrealEd/Classes/Settings/EditorExperimentalSettings.h
--- a/Engine/Source/Editor/UnrealEd/Classes/Settings/EditorExperimentalSettings.h
+++ b/Engine/Source/Editor/UnrealEd/Classes/Settings/EditorExperimentalSettings.h
@@ -89,6 +89,10 @@ public:
 	UPROPERTY(EditAnywhere, config, Category = Tools, meta = (DisplayName = "Visual Logger"))
 	bool bVisualLogger;
 
+	/** Enables Gameplay Debugger */
+	UPROPERTY(EditAnywhere, config, Category = Tools, meta = (DisplayName = "Gameplay Debugger"))
+	bool bGameplayDebugger;
+
 	/**
 	 * Returns an event delegate that is executed when a setting has changed.
 	 *
