Minor cleanups for GameplayDebugger's and LogVisualizer's settings.

diff --git a/Engine/Source/Developer/GameplayDebugger/Classes/GameplayDebuggerSettings.h b/Engine/Source/Developer/GameplayDebugger/Classes/GameplayDebuggerSettings.h
--- a/Engine/Source/Developer/GameplayDebugger/Classes/GameplayDebuggerSettings.h
+++ b/Engine/Source/Developer/GameplayDebugger/Classes/GameplayDebuggerSettings.h
@@ -8,6 +8,8 @@ GameplayDebuggerSettings.h: Declares the UGameplayDebuggerSettings class.
 
 #include "GameplayDebuggerSettings.generated.h"
 
+#define ADD_LEVEL_EDITOR_EXTENSIONS 0
+
 USTRUCT()
 struct FGDTCustomViewNames
 {
@@ -41,7 +43,10 @@ public:
 	// End UObject Interface
 #endif
 
-	uint32 GetSettings() 
+	DECLARE_EVENT_OneParam(UGameplayDebuggerSettings, FSettingChangedEvent, FName /*PropertyName*/);
+	FSettingChangedEvent& OnSettingChanged() { return SettingChangedEvent; }
+
+	uint32 GetSettings()
 	{
 		uint32 Settings = 0;
 		Settings = OverHead ? Settings | (1 << EAIDebugDrawDataView::OverHead) : Settings & ~(1 << EAIDebugDrawDataView::OverHead);
@@ -60,6 +65,12 @@ public:
 
 	const FGDTCustomViewNames& GetCustomViewNames() { return CustomViewNames; }
 
+#if ADD_LEVEL_EDITOR_EXTENSIONS
+	/** Experimental setting to extend viewport menu in Simulate, to have quick access to GameplayDebugger settings. */
+	UPROPERTY(EditAnywhere, config, Category = "GameplayDebugger")
+	bool bExtendViewportMenu;
+#endif
+
 protected:
 	UPROPERTY(EditAnywhere, config, Category = "GameplayDebugger")
 	FGDTCustomViewNames CustomViewNames;
@@ -99,4 +110,7 @@ protected:
 
 	UPROPERTY(EditAnywhere, config, Category = "GameplayDebugger")
 	bool GameView5;
+
+	// Holds an event delegate that is executed when a setting has changed.
+	FSettingChangedEvent SettingChangedEvent;
 };
diff --git a/Engine/Source/Developer/GameplayDebugger/Private/GameplayDebugger.cpp b/Engine/Source/Developer/GameplayDebugger/Private/GameplayDebugger.cpp
--- a/Engine/Source/Developer/GameplayDebugger/Private/GameplayDebugger.cpp
+++ b/Engine/Source/Developer/GameplayDebugger/Private/GameplayDebugger.cpp
@@ -14,8 +14,6 @@
 #include "LevelEditor.h"
 #endif // WITH_EDITOR
 
-#define ADD_LEVEL_EDITOR_EXTENSIONS 0
-
 #define LOCTEXT_NAMESPACE "FGameplayDebugger"
 FGameplayDebuggerSettings GameplayDebuggerSettings(class AGameplayDebuggingReplicator* Replicator)
 {
@@ -39,7 +37,7 @@ public:
 	TSharedRef<FExtender> OnExtendLevelEditorViewMenu(const TSharedRef<FUICommandList> CommandList);
 	void CreateSnappingOptionsMenu(FMenuBuilder& Builder);
 	void CreateSettingSubMenu(FMenuBuilder& Builder);
-	void HandleExperimentalSettingChanged(FName PropertyName);
+	void HandleSettingChanged(FName PropertyName);
 #endif
 
 	TArray<TWeakObjectPtr<AGameplayDebuggingReplicator> >& GetAllReplicators(UWorld* InWorld);
@@ -76,8 +74,6 @@ void FGameplayDebugger::StartupModule()
 		GEngine->OnLevelActorAdded().AddRaw(this, &FGameplayDebugger::OnLevelActorAdded);
 		GEngine->OnLevelActorDeleted().AddRaw(this, &FGameplayDebugger::OnLevelActorDeleted);
 
-		UEditorExperimentalSettings::StaticClass()->GetDefaultObject<UEditorExperimentalSettings>()->OnSettingChanged().AddRaw(this, &FGameplayDebugger::HandleExperimentalSettingChanged);
-
 		ISettingsModule* SettingsModule = FModuleManager::GetModulePtr<ISettingsModule>("Settings");
 		if (SettingsModule != nullptr)
 		{
@@ -87,8 +83,12 @@ void FGameplayDebugger::StartupModule()
 				UGameplayDebuggerSettings::StaticClass()->GetDefaultObject()
 				);
 		}
+
+		UGameplayDebuggerSettings::StaticClass()->GetDefaultObject<UGameplayDebuggerSettings>()->OnSettingChanged().AddRaw(this, &FGameplayDebugger::HandleSettingChanged);
+
 #	if ADD_LEVEL_EDITOR_EXTENSIONS
-		if (!IsRunningCommandlet())
+		const bool bExtendViewportMenu = UGameplayDebuggerSettings::StaticClass()->GetDefaultObject<UGameplayDebuggerSettings>()->bExtendViewportMenu;
+		if (!IsRunningCommandlet() && bExtendViewportMenu)
 		{
 			// Register the extension with the level editor
 			{
@@ -106,6 +106,41 @@ void FGameplayDebugger::StartupModule()
 #endif
 }
 
+#if WITH_EDITOR
+void FGameplayDebugger::HandleSettingChanged(FName PropertyName)
+{
+#if ADD_LEVEL_EDITOR_EXTENSIONS
+	const bool bExtendViewportMenu = UGameplayDebuggerSettings::StaticClass()->GetDefaultObject<UGameplayDebuggerSettings>()->bExtendViewportMenu;
+	if (PropertyName == TEXT("bExtendViewportMenu"))
+	{
+		if (bExtendViewportMenu)
+		{
+			// Register the extension with the level editor
+			{
+				ViewMenuExtender = FLevelEditorModule::FLevelEditorMenuExtender::CreateRaw(this, &FGameplayDebugger::OnExtendLevelEditorViewMenu);
+				FLevelEditorModule* LevelEditor = FModuleManager::LoadModulePtr<FLevelEditorModule>(TEXT("LevelEditor"));
+				if (LevelEditor)
+				{
+					LevelEditor->GetAllLevelViewportOptionsMenuExtenders().Add(ViewMenuExtender);
+				}
+			}
+		}
+		else
+		{
+			// Unregister the level editor extensions
+			{
+				FLevelEditorModule* LevelEditor = FModuleManager::LoadModulePtr<FLevelEditorModule>(TEXT("LevelEditor"));
+				if (LevelEditor)
+				{
+					LevelEditor->GetAllLevelViewportOptionsMenuExtenders().Remove(ViewMenuExtender);
+				}
+			}
+		}
+	}
+#endif
+}
+#endif
+
 // This function may be called during shutdown to clean up your module.  For modules that support dynamic reloading,
 // we call this function before unloading the module.
 void FGameplayDebugger::ShutdownModule()
@@ -119,7 +154,7 @@ void FGameplayDebugger::ShutdownModule()
 		GEngine->OnLevelActorAdded().RemoveAll(this);
 		GEngine->OnLevelActorDeleted().RemoveAll(this);
 
-		UEditorExperimentalSettings::StaticClass()->GetDefaultObject<UEditorExperimentalSettings>()->OnSettingChanged().RemoveAll(this);
+		UGameplayDebuggerSettings::StaticClass()->GetDefaultObject<UGameplayDebuggerSettings>()->OnSettingChanged().RemoveAll(this);
 
 		ISettingsModule* SettingsModule = FModuleManager::GetModulePtr<ISettingsModule>("Settings");
 		if (SettingsModule != nullptr)
@@ -147,36 +182,6 @@ void FGameplayDebugger::ShutdownModule()
 
 #if WITH_EDITOR
 
-void FGameplayDebugger::HandleExperimentalSettingChanged(FName PropertyName)
-{
-	if (PropertyName == TEXT("bGameplayDebugger"))
-	{
-		if (UEditorExperimentalSettings::StaticClass()->GetDefaultObject<UEditorExperimentalSettings>()->bGameplayDebugger == false)
-		{
-			ISettingsModule* SettingsModule = FModuleManager::GetModulePtr<ISettingsModule>("Settings");
-			if (SettingsModule != nullptr)
-			{
-				SettingsModule->UnregisterSettings("Editor", "General", "GameplayDebugger");
-			}
-		}
-		else
-		{
-			ISettingsModule* SettingsModule = FModuleManager::GetModulePtr<ISettingsModule>("Settings");
-			if (SettingsModule != nullptr)
-			{
-				SettingsModule->RegisterSettings("Editor", "General", "GameplayDebugger",
-					LOCTEXT("AIToolsSettingsName", "Gameplay Debugger"),
-					LOCTEXT("AIToolsSettingsDescription", "General settings for UE4 AI Tools."),
-					UGameplayDebuggerSettings::StaticClass()->GetDefaultObject()
-					);
-			}
-		}
-	}
-	//FLevelEditorModule& LevelEditorModule = FModuleManager::GetModuleChecked<FLevelEditorModule>("LevelEditor");
-	//TSharedPtr<FTabManager> LevelEditorTabManager = LevelEditorModule.GetLevelEditorTabManager();
-	//LevelEditorTabManager->UpdateMainMenu(true);
-}
-
 void FGameplayDebugger::CreateSettingSubMenu(FMenuBuilder& Builder)
 {
 #if !(UE_BUILD_SHIPPING || UE_BUILD_TEST)
@@ -212,17 +217,20 @@ void FGameplayDebugger::CreateSnappingOptionsMenu(FMenuBuilder& Builder)
 
 TSharedRef<FExtender> FGameplayDebugger::OnExtendLevelEditorViewMenu(const TSharedRef<FUICommandList> CommandList)
 {
-#if !(UE_BUILD_SHIPPING || UE_BUILD_TEST)
+#if !(UE_BUILD_SHIPPING || UE_BUILD_TEST) && WITH_EDITOR
 	TSharedRef<FExtender> Extender(new FExtender());
+	if (GEditor && GEditor->bIsSimulatingInEditor)
+	{
 
-	Extender->AddMenuExtension(
-		"LevelViewportLayouts",
-		EExtensionHook::After,
-		NULL,
-		FMenuExtensionDelegate::CreateRaw(this, &FGameplayDebugger::CreateSnappingOptionsMenu));
+		Extender->AddMenuExtension(
+			"LevelViewportLayouts",
+			EExtensionHook::After,
+			NULL,
+			FMenuExtensionDelegate::CreateRaw(this, &FGameplayDebugger::CreateSnappingOptionsMenu));
 
-	return Extender;
+	}
 #endif
+	return Extender;
 }
 #endif
 
diff --git a/Engine/Source/Developer/GameplayDebugger/Private/GameplayDebuggerSettings.cpp b/Engine/Source/Developer/GameplayDebugger/Private/GameplayDebuggerSettings.cpp
--- a/Engine/Source/Developer/GameplayDebugger/Private/GameplayDebuggerSettings.cpp
+++ b/Engine/Source/Developer/GameplayDebugger/Private/GameplayDebuggerSettings.cpp
@@ -36,6 +36,11 @@ UGameplayDebuggerSettings::UGameplayDebuggerSettings(const FObjectInitializer& O
 		GConfig->GetBool(TEXT("GameplayDebuggerSettings"), TEXT("GameView3"), GameView3, GEngineIni);
 		GConfig->GetBool(TEXT("GameplayDebuggerSettings"), TEXT("GameView4"), GameView4, GEngineIni);
 		GConfig->GetBool(TEXT("GameplayDebuggerSettings"), TEXT("GameView5"), GameView5, GEngineIni);
+
+#if ADD_LEVEL_EDITOR_EXTENSIONS
+		bExtendViewportMenu = false;
+		GConfig->GetBool(TEXT("GameplayDebuggerSettings"), TEXT("bExtendViewportMenu"), bExtendViewportMenu, GEngineIni);
+#endif
 	}
 }
 
@@ -88,5 +93,7 @@ void UGameplayDebuggerSettings::PostEditChangeProperty(FPropertyChangedEvent& Pr
 #undef UPDATE_GAMEVIEW_DISPLAYNAME
 
 	GEditor->SaveEditorUserSettings();
+
+	SettingChangedEvent.Broadcast(Name);
 }
 #endif
diff --git a/Engine/Source/Developer/LogVisualizer/Public/LogVisualizerSettings.h b/Engine/Source/Developer/LogVisualizer/Public/LogVisualizerSettings.h
--- a/Engine/Source/Developer/LogVisualizer/Public/LogVisualizerSettings.h
+++ b/Engine/Source/Developer/LogVisualizer/Public/LogVisualizerSettings.h
@@ -13,6 +13,9 @@ class LOGVISUALIZER_API ULogVisualizerSettings : public UObject
 {
 	GENERATED_UCLASS_BODY()
 public:
+	DECLARE_EVENT_OneParam(ULogVisualizerSettings, FSettingChangedEvent, FName /*PropertyName*/);
+	FSettingChangedEvent& OnSettingChanged() { return SettingChangedEvent; }
+
 	/**Whether to show trivial logs, i.e. the ones with only one entry.*/
 	UPROPERTY(EditAnywhere, config, Category = "VisualLogger")
 	bool bIgnoreTrivialLogs;
@@ -21,21 +24,18 @@ public:
 	UPROPERTY(EditAnywhere, config, Category = "VisualLogger", meta = (EditCondition = "bIgnoreTrivialLogs", ClampMin = "0", ClampMax = "10", UIMin = "0", UIMax = "10"))
 	int32 TrivialLogsThreshold;
 
-	/**Whether to show the recent data or not.*/
-	UPROPERTY(EditAnywhere, config, Category = "VisualLogger")
+	/**Whether to show the recent data or not. Property disabled for now.*/
+	UPROPERTY(VisibleAnywhere, config, Category = "VisualLogger")
 	bool bStickToRecentData;
 
-	/**Whether to show histogram labels inside graph or outside.*/
-	UPROPERTY(EditAnywhere, config, Category = "VisualLogger")
+	/**Whether to show histogram labels inside graph or outside. Property disabled for now.*/
+	UPROPERTY(VisibleAnywhere, config, Category = "VisualLogger")
 	bool bShowHistogramLabelsOutside;
 
 	/** Camera distance used to setup location during reaction on log item double click */
 	UPROPERTY(EditAnywhere, config, Category = "VisualLogger", meta = (ClampMin = "10", ClampMax = "1000", UIMin = "10", UIMax = "1000"))
 	float DefaultCameraDistance;
 
-	DECLARE_EVENT_OneParam(ULogVisualizerSettings, FSettingChangedEvent, FName /*PropertyName*/);
-	FSettingChangedEvent& OnSettingChanged() { return SettingChangedEvent; }
-
 	/**Whether to search/filter categories or to get text vlogs into account too */
 	UPROPERTY(EditAnywhere, config, Category = "VisualLogger")
 	bool bSearchInsideLogs;
diff --git a/Engine/Source/Editor/UnrealEd/Classes/Settings/EditorExperimentalSettings.h b/Engine/Source/Editor/UnrealEd/Classes/Settings/EditorExperimentalSettings.h
--- a/Engine/Source/Editor/UnrealEd/Classes/Settings/EditorExperimentalSettings.h
+++ b/Engine/Source/Editor/UnrealEd/Classes/Settings/EditorExperimentalSettings.h
@@ -85,10 +85,6 @@ public:
 	UPROPERTY(EditAnywhere, config, Category = Blueprints, meta = (DisplayName = "Enable Blueprint Merge Tool"))
 	bool bEnableBlueprintMergeTool;
 
-	/** Enables Gameplay Debugger */
-	UPROPERTY(EditAnywhere, config, Category = Tools, meta = (DisplayName = "Gameplay Debugger"))
-	bool bGameplayDebugger;
-
 	/** Enables GetContent button in asset browser */
 	UPROPERTY(EditAnywhere, config, Category = Tools, meta = (DisplayName = "Get Feature Content (Requires content browser reopening)"))
 	bool bGetFeatureContent;
