Added a way to extend GameplayDebugger's controller class, to set custom keyboard bindings, etc.

diff --git a/Engine/Config/BaseEngine.ini b/Engine/Config/BaseEngine.ini
--- a/Engine/Config/BaseEngine.ini
+++ b/Engine/Config/BaseEngine.ini
@@ -2074,6 +2074,7 @@ StartingBeaconPort=15000
 MaxEQSQueries=5
 DebugComponentClassName="/Script/GameplayDebugger.GameplayDebuggingComponent"
 DebugComponentHUDClassName="/Script/GameplayDebugger.GameplayDebuggingHUDComponent"
+DebugComponentControllerClassName = "/Script/GameplayDebugger.GameplayDebuggingControllerComponent"
 
 ; Note: These should be in IOSEngine, but are placed here because the Settings editor UI cannot write to the Game IOSEngine, meaning settings in Engine/IOSEngine are non-overridable
 [/Script/IOSRuntimeSettings.IOSRuntimeSettings]
diff --git a/Engine/Source/Developer/GameplayDebugger/Classes/GameplayDebuggingReplicator.h b/Engine/Source/Developer/GameplayDebugger/Classes/GameplayDebuggingReplicator.h
--- a/Engine/Source/Developer/GameplayDebugger/Classes/GameplayDebuggingReplicator.h
+++ b/Engine/Source/Developer/GameplayDebugger/Classes/GameplayDebuggingReplicator.h
@@ -23,6 +23,9 @@ class GAMEPLAYDEBUGGER_API AGameplayDebuggingReplicator : public AActor
 	UPROPERTY(config)
 	FString DebugComponentHUDClassName;
 
+	UPROPERTY(config)
+	FString DebugComponentControllerClassName;
+
 	UPROPERTY(config)
 	int32 MaxEQSQueries;
 
@@ -135,6 +138,6 @@ private:
 
 	TWeakObjectPtr<UClass> DebugComponentClass;
 	TWeakObjectPtr<UClass> DebugComponentHUDClass;
-
+	TWeakObjectPtr<UClass> DebugComponentControllerClass;
 	TWeakObjectPtr<AGameplayDebuggingHUDComponent>	DebugRenderer;
 };
diff --git a/Engine/Source/Developer/GameplayDebugger/Private/GameplayDebuggingReplicator.cpp b/Engine/Source/Developer/GameplayDebugger/Private/GameplayDebuggingReplicator.cpp
--- a/Engine/Source/Developer/GameplayDebugger/Private/GameplayDebuggingReplicator.cpp
+++ b/Engine/Source/Developer/GameplayDebugger/Private/GameplayDebuggingReplicator.cpp
@@ -372,7 +372,12 @@ void AGameplayDebuggingReplicator::CreateTool()
 		UGameplayDebuggingControllerComponent*  GDC = FindComponentByClass<UGameplayDebuggingControllerComponent>();
 		if (!GDC)
 		{
-			GDC = ConstructObject<UGameplayDebuggingControllerComponent>(UGameplayDebuggingControllerComponent::StaticClass(), this);
+			DebugComponentControllerClass = StaticLoadClass(UGameplayDebuggingControllerComponent::StaticClass(), NULL, *DebugComponentControllerClassName, NULL, LOAD_None, NULL);
+			if (!DebugComponentControllerClass.IsValid())
+			{
+				DebugComponentControllerClass = AGameplayDebuggingHUDComponent::StaticClass();
+			}
+			GDC = ConstructObject<UGameplayDebuggingControllerComponent>(DebugComponentControllerClass.Get(), this);
 			GDC->SetPlayerOwner(LocalPlayerOwner);
 			GDC->RegisterComponent();
 		}
