Fixed Clang compilation (missing includes in GameplayDebugger.h)

diff --git a/Engine/Source/Developer/GameplayDebugger/Private/GameplayDebuggerPrivate.h b/Engine/Source/Developer/GameplayDebugger/Private/GameplayDebuggerPrivate.h
--- a/Engine/Source/Developer/GameplayDebugger/Private/GameplayDebuggerPrivate.h
+++ b/Engine/Source/Developer/GameplayDebugger/Private/GameplayDebuggerPrivate.h
@@ -24,6 +24,8 @@
 #include "Components/CapsuleComponent.h"
 #include "Engine/Engine.h"
 
+#include "RenderResource.h"
+#include "Shader.h"
 #include "Engine/Canvas.h"
 #include "EngineUtils.h"
 
