Fixed on compilation warning due to using deprecated function I've missed #UE4

diff --git a/Engine/Source/Developer/GameplayDebugger/Private/GameplayDebuggingComponent.cpp b/Engine/Source/Developer/GameplayDebugger/Private/GameplayDebuggingComponent.cpp
--- a/Engine/Source/Developer/GameplayDebugger/Private/GameplayDebuggingComponent.cpp
+++ b/Engine/Source/Developer/GameplayDebugger/Private/GameplayDebuggingComponent.cpp
@@ -928,7 +928,7 @@ ARecastNavMesh* UGameplayDebuggingComponent::GetNavData()
 	APawn* TargetPawn = Cast<APawn>(TargetActor);
 	if (TargetPawn != NULL)
 	{
-		const FNavAgentProperties& NavAgentProperties = TargetPawn->GetNavAgentProperties();
+		const FNavAgentProperties& NavAgentProperties = TargetPawn->GetNavAgentPropertiesRef();
 		return Cast<ARecastNavMesh>(NavSys->GetNavDataForProps(NavAgentProperties));
 	}
 
