extended gameplay debugger's basic pawn data

diff --git a/Engine/Source/Developer/GameplayDebugger/Classes/GameplayDebuggingComponent.h b/Engine/Source/Developer/GameplayDebugger/Classes/GameplayDebuggingComponent.h
--- a/Engine/Source/Developer/GameplayDebugger/Classes/GameplayDebuggingComponent.h
+++ b/Engine/Source/Developer/GameplayDebugger/Classes/GameplayDebuggingComponent.h
@@ -61,6 +61,33 @@ class GAMEPLAYDEBUGGER_API UGameplayDebuggingComponent : public UPrimitiveCompon
 	UPROPERTY(Replicated)
 	FString DebugIcon;
 
+	UPROPERTY(Replicated)
+	FString MovementBaseInfo;
+
+	UPROPERTY(Replicated)
+	FString MovementModeInfo;
+
+	UPROPERTY(Replicated)
+	FString PathFollowingInfo;
+
+	UPROPERTY(Replicated)
+	FString CurrentAITask;
+
+	UPROPERTY(Replicated)
+	FString CurrentAIState;
+
+	UPROPERTY(Replicated)
+	FString CurrentAIAssets;
+
+	UPROPERTY(Replicated)
+	FString NavDataInfo;
+
+	UPROPERTY(Replicated)
+	FString AbilityInfo;
+
+	UPROPERTY(Replicated)
+	FString MontageInfo;
+
 	UPROPERTY(Replicated)
 	FString BrainComponentName;
 
@@ -76,14 +103,10 @@ class GAMEPLAYDEBUGGER_API UGameplayDebuggingComponent : public UPrimitiveCompon
 	UPROPERTY(Replicated)
 	TArray<FVector> PathPoints;
 
-	/** Begin path replication data */
 	UPROPERTY(ReplicatedUsing = OnRep_PathCorridorData)
 	TArray<uint8> PathCorridorData;
 	
 	TArray<TArray<FVector> > PathCorridorPolygons;
-
-	UPROPERTY(Replicated)
-	FString PathErrorString;
 	/** End path replication data*/
 	
 	UPROPERTY(ReplicatedUsing = OnRep_UpdateNavmesh)
@@ -103,6 +126,21 @@ class GAMEPLAYDEBUGGER_API UGameplayDebuggingComponent : public UPrimitiveCompon
 	UPROPERTY(Replicated)
 	FVector SensingComponentLocation;
 
+	UPROPERTY(Replicated)
+	int32 NextPathPointIndex;
+
+	UPROPERTY(Replicated)
+	uint32 bIsUsingPathFollowing : 1;
+
+	UPROPERTY(Replicated)
+	uint32 bIsUsingCharacter : 1;
+
+	UPROPERTY(Replicated)
+	uint32 bIsUsingBehaviorTree : 1;
+
+	UPROPERTY(Replicated)
+	uint32 bIsUsingAbilities : 1;
+
 	uint32 bDrawEQSLabels:1;
 	uint32 bDrawEQSFailedItems : 1;
 
@@ -201,6 +239,12 @@ protected:
 	virtual void CollectBasicData();
 	virtual void CollectBehaviorTreeData();
 
+	virtual void CollectBasicMovementData(APawn* MyPawn);
+	virtual void CollectBasicPathData(APawn* MyPawn);
+	virtual void CollectBasicBehaviorData(APawn* MyPawn);
+	virtual void CollectBasicAbilityData(APawn* MyPawn);
+	virtual void CollectBasicAnimationData(APawn* MyPawn);
+
 	FNavPathWeakPtr CurrentPath;
 
 	uint32 bEnabledTargetSelection : 1;
diff --git a/Engine/Source/Developer/GameplayDebugger/GameplayDebugger.Build.cs b/Engine/Source/Developer/GameplayDebugger/GameplayDebugger.Build.cs
--- a/Engine/Source/Developer/GameplayDebugger/GameplayDebugger.Build.cs
+++ b/Engine/Source/Developer/GameplayDebugger/GameplayDebugger.Build.cs
@@ -11,6 +11,7 @@ namespace UnrealBuildTool.Rules
 				    "Developer/GameplayDebugger/Public",
 				    "Developer/AIModule/Public",
                     "Developer/Settings/Public",
+                    "Runtime/GameplayAbilities/Public",
 				    // ... add public include paths required here ...
 			    }
             );
@@ -36,6 +37,7 @@ namespace UnrealBuildTool.Rules
                     "ShaderCore",
                     "Settings",
                     "AIModule",  // it have to be here for now. It'll be changed to remove any dependency to AIModule in future
+                    "GameplayAbilities",
 				}
                 );
 
diff --git a/Engine/Source/Developer/GameplayDebugger/Private/GameplayDebuggingComponent.cpp b/Engine/Source/Developer/GameplayDebugger/Private/GameplayDebuggingComponent.cpp
--- a/Engine/Source/Developer/GameplayDebugger/Private/GameplayDebuggingComponent.cpp
+++ b/Engine/Source/Developer/GameplayDebugger/Private/GameplayDebuggingComponent.cpp
@@ -10,7 +10,8 @@
 #include "EnvironmentQuery/EQSRenderingComponent.h"
 #include "EnvironmentQuery/EnvQueryTest.h"
 #include "BehaviorTree/BlackboardComponent.h"
-
+#include "BehaviorTree/BehaviorTreeComponent.h"
+#include "AbilitySystemComponent.h"
 #include "AIController.h"
 #include "BrainComponent.h"
 #include "BehaviorTreeDelegates.h"
@@ -184,29 +185,55 @@ void UGameplayDebuggingComponent::GetLifetimeReplicatedProps( TArray< FLifetimeP
 {
 	Super::GetLifetimeReplicatedProps( OutLifetimeProps );
 #if !(UE_BUILD_SHIPPING || UE_BUILD_TEST)
-	DOREPLIFETIME( UGameplayDebuggingComponent, ReplicateViewDataCounters );
-	DOREPLIFETIME( UGameplayDebuggingComponent, ShowExtendedInformatiomCounter );
-	DOREPLIFETIME( UGameplayDebuggingComponent, ControllerName )
-	DOREPLIFETIME( UGameplayDebuggingComponent, PawnName );
-	DOREPLIFETIME( UGameplayDebuggingComponent, DebugIcon );
-	DOREPLIFETIME( UGameplayDebuggingComponent, PawnClass );
+	DOREPLIFETIME(UGameplayDebuggingComponent, ReplicateViewDataCounters);
+	DOREPLIFETIME(UGameplayDebuggingComponent, ShowExtendedInformatiomCounter);
+	DOREPLIFETIME(UGameplayDebuggingComponent, ControllerName)
+	DOREPLIFETIME(UGameplayDebuggingComponent, PawnName);
+	DOREPLIFETIME(UGameplayDebuggingComponent, DebugIcon);
+	DOREPLIFETIME(UGameplayDebuggingComponent, PawnClass);
+
+	DOREPLIFETIME(UGameplayDebuggingComponent, bIsUsingCharacter);
+	DOREPLIFETIME(UGameplayDebuggingComponent, bIsUsingPathFollowing);
+	DOREPLIFETIME(UGameplayDebuggingComponent, bIsUsingBehaviorTree);
+	DOREPLIFETIME(UGameplayDebuggingComponent, bIsUsingAbilities);
+
+	if (bIsUsingCharacter)
+	{
+		DOREPLIFETIME(UGameplayDebuggingComponent, MovementBaseInfo);
+		DOREPLIFETIME(UGameplayDebuggingComponent, MovementModeInfo);
+		DOREPLIFETIME(UGameplayDebuggingComponent, MontageInfo);
+	}
 
-	DOREPLIFETIME(UGameplayDebuggingComponent, BrainComponentName);
-	DOREPLIFETIME(UGameplayDebuggingComponent, BrainComponentString);
-	DOREPLIFETIME(UGameplayDebuggingComponent, BlackboardRepData);
+	if (bIsUsingPathFollowing)
+	{
+		DOREPLIFETIME(UGameplayDebuggingComponent, PathFollowingInfo);
+		DOREPLIFETIME(UGameplayDebuggingComponent, NavDataInfo);
+		DOREPLIFETIME(UGameplayDebuggingComponent, PathPoints);
+		DOREPLIFETIME(UGameplayDebuggingComponent, PathCorridorData);
+		DOREPLIFETIME(UGameplayDebuggingComponent, NextPathPointIndex);
+	}
 
-	DOREPLIFETIME( UGameplayDebuggingComponent, PathPoints );
-	DOREPLIFETIME( UGameplayDebuggingComponent, PathErrorString );
+	if (bIsUsingBehaviorTree)
+	{
+		DOREPLIFETIME(UGameplayDebuggingComponent, CurrentAITask);
+		DOREPLIFETIME(UGameplayDebuggingComponent, CurrentAIState);
+		DOREPLIFETIME(UGameplayDebuggingComponent, CurrentAIAssets);
+	}
 
-	DOREPLIFETIME( UGameplayDebuggingComponent, NavmeshRepData );
+	if (bIsUsingAbilities)
+	{
+		DOREPLIFETIME(UGameplayDebuggingComponent, AbilityInfo);
+	}
 
-	DOREPLIFETIME( UGameplayDebuggingComponent, TargetActor );
+	DOREPLIFETIME(UGameplayDebuggingComponent, BrainComponentName);
+	DOREPLIFETIME(UGameplayDebuggingComponent, BrainComponentString);
+	DOREPLIFETIME(UGameplayDebuggingComponent, BlackboardRepData);
 
+	DOREPLIFETIME(UGameplayDebuggingComponent, NavmeshRepData);
+	DOREPLIFETIME(UGameplayDebuggingComponent, TargetActor);
 	DOREPLIFETIME(UGameplayDebuggingComponent, EQSRepData);
 
 	DOREPLIFETIME(UGameplayDebuggingComponent, SensingComponentLocation);
-
-	DOREPLIFETIME(UGameplayDebuggingComponent, PathCorridorData);
 #endif //!(UE_BUILD_SHIPPING || UE_BUILD_TEST)
 }
 
@@ -383,17 +410,28 @@ void UGameplayDebuggingComponent::CollectDataToReplicate(bool bCollectExtendedDa
 void UGameplayDebuggingComponent::CollectBasicData()
 {
 #if !(UE_BUILD_SHIPPING || UE_BUILD_TEST)
-	const APawn* MyPawn = Cast<APawn>(GetSelectedActor());
+	APawn* MyPawn = Cast<APawn>(GetSelectedActor());
 	PawnName = MyPawn->GetHumanReadableName();
 	PawnClass = MyPawn->GetClass()->GetName();
-	const AAIController* MyController = Cast<const AAIController>(MyPawn->Controller);
+	AAIController* MyController = Cast<AAIController>(MyPawn->Controller);
 
-	if (MyController != NULL)
+	bIsUsingCharacter = MyPawn->IsA(ACharacter::StaticClass());
+	bIsUsingPathFollowing = false;
+	bIsUsingBehaviorTree = false;
+	bIsUsingAbilities = false;
+
+	if (MyController)
 	{
 		if (MyController->IsPendingKill() == false)
 		{
 			ControllerName = MyController->GetName();
 			DebugIcon = MyController->GetDebugIcon();
+
+			CollectBasicMovementData(MyPawn);
+			CollectBasicPathData(MyPawn);
+			CollectBasicBehaviorData(MyPawn);
+			CollectBasicAbilityData(MyPawn);
+			CollectBasicAnimationData(MyPawn);
 		}
 		else
 		{
@@ -407,6 +445,141 @@ void UGameplayDebuggingComponent::CollectBasicData()
 #endif //!(UE_BUILD_SHIPPING || UE_BUILD_TEST)
 }
 
+void UGameplayDebuggingComponent::CollectBasicMovementData(APawn* MyPawn)
+{
+	UCharacterMovementComponent* CharMovement = Cast<UCharacterMovementComponent>(MyPawn->GetMovementComponent());
+	if (CharMovement)
+	{
+		UPrimitiveComponent* FloorComponent = MyPawn->GetMovementBase();
+		AActor* FloorActor = FloorComponent ? FloorComponent->GetOwner() : nullptr;
+		MovementBaseInfo = FloorComponent ? FString::Printf(TEXT("%s.%s"), *GetNameSafe(FloorActor), *FloorComponent->GetName()) : TEXT("None");
+
+		MovementModeInfo = CharMovement->GetMovementName();
+	}
+	else
+	{
+		MovementBaseInfo = TEXT("");
+		MovementModeInfo = TEXT("");
+	}
+}
+
+void UGameplayDebuggingComponent::CollectBasicPathData(APawn* MyPawn)
+{
+	UNavigationSystem* NavSys = UNavigationSystem::GetCurrent(MyPawn->GetWorld());
+	AAIController* MyAIController = Cast<AAIController>(MyPawn->GetController());
+
+	const ANavigationData* NavData = NavSys->GetNavDataForProps(MyAIController->GetNavAgentPropertiesRef());
+	NavDataInfo = NavData->NavDataConfig.Name.ToString();
+
+	UPathFollowingComponent* PFC = MyAIController->GetPathFollowingComponent();
+	bIsUsingPathFollowing = (PFC != nullptr);
+
+	if (PFC)
+	{
+		TArray<FString> Tokens;
+		TArray<EPathFollowingDebugTokens::Type> Flags;
+		if (PFC)
+		{
+			PFC->GetDebugStringTokens(Tokens, Flags);
+		}
+
+		PathFollowingInfo = FString();
+		for (int32 Idx = 0; Idx < Tokens.Num(); Idx++)
+		{
+			switch (Flags[Idx])
+			{
+			case EPathFollowingDebugTokens::Description:
+				PathFollowingInfo += Tokens[Idx];
+				break;
+
+			case EPathFollowingDebugTokens::ParamName:
+				PathFollowingInfo += TEXT(", {yellow}");
+				PathFollowingInfo += Tokens[Idx];
+				PathFollowingInfo += TEXT(":");
+				break;
+
+			case EPathFollowingDebugTokens::PassedValue:
+				PathFollowingInfo += TEXT("{yellow}");
+				PathFollowingInfo += Tokens[Idx];
+				break;
+
+			case EPathFollowingDebugTokens::FailedValue:
+				PathFollowingInfo += TEXT("{orange}");
+				PathFollowingInfo += Tokens[Idx];
+				break;
+
+			default:
+				break;
+			}
+		}
+	}
+	else
+	{
+		PathFollowingInfo = TEXT("");
+	}
+}
+
+void UGameplayDebuggingComponent::CollectBasicBehaviorData(APawn* MyPawn)
+{
+	AAIController* MyAIController = Cast<AAIController>(MyPawn->GetController());
+	UBehaviorTreeComponent* BTC = MyAIController ? Cast<UBehaviorTreeComponent>(MyAIController->BrainComponent) : nullptr;
+	bIsUsingBehaviorTree = (BTC != nullptr);
+
+	if (BTC)
+	{
+		CurrentAITask = BTC->DescribeActiveTasks();
+		CurrentAIState = BTC->IsRunning() ? TEXT("Running") : BTC->IsPaused() ? TEXT("Paused") : TEXT("Inactive");
+		CurrentAIAssets = BTC->DescribeActiveTrees();
+	}
+	else
+	{
+		CurrentAITask = TEXT("");
+		CurrentAIState = TEXT("");
+		CurrentAIAssets = TEXT("");
+	}
+}
+
+void UGameplayDebuggingComponent::CollectBasicAbilityData(APawn* MyPawn)
+{
+	UAbilitySystemComponent* AbilityComp = MyPawn->FindComponentByClass<UAbilitySystemComponent>();
+	bIsUsingAbilities = (AbilityComp != nullptr);
+
+	int32 NumActive = 0;
+	if (AbilityComp)
+	{
+		AbilityInfo = TEXT("");
+
+		for (const FGameplayAbilitySpec& AbilitySpec : AbilityComp->GetActivatableAbilities())
+		{
+			if (AbilitySpec.Ability && AbilitySpec.IsActive())
+			{
+				if (NumActive)
+				{
+					AbilityInfo += TEXT(", ");
+				}
+
+				UClass* AbilityClass = AbilitySpec.Ability->GetClass();
+				FString AbClassName = GetNameSafe(AbilityClass);
+				AbClassName.RemoveFromEnd(TEXT("_c"));
+
+				AbilityInfo += AbClassName;
+				NumActive++;
+			}
+		}
+	}
+
+	if (NumActive == 0)
+	{
+		AbilityInfo = TEXT("None");
+	}
+}
+
+void UGameplayDebuggingComponent::CollectBasicAnimationData(APawn* MyPawn)
+{
+	ACharacter* MyChar = Cast<ACharacter>(MyPawn);
+	MontageInfo = MyChar ? GetNameSafe(MyChar->GetCurrentMontage()) : TEXT("");
+}
+
 void UGameplayDebuggingComponent::CollectBehaviorTreeData()
 {
 #if !(UE_BUILD_SHIPPING || UE_BUILD_TEST)
@@ -478,13 +651,14 @@ void UGameplayDebuggingComponent::CollectPathData()
 {
 #if !(UE_BUILD_SHIPPING || UE_BUILD_TEST)
 	APawn* MyPawn = Cast<APawn>(GetSelectedActor());
-	PathErrorString.Empty();
 
 	bool bRefreshRendering = false;
 	if (AAIController *MyController = Cast<AAIController>(MyPawn->Controller))
 	{
 		if (MyController->PathFollowingComponent && MyController->PathFollowingComponent->GetPath().IsValid())
 		{
+			NextPathPointIndex = MyController->PathFollowingComponent->GetNextPathIndex();
+
 			const FNavPathSharedPtr& NewPath = MyController->PathFollowingComponent->GetPath();
 			if (!CurrentPath.HasSameObject(NewPath.Get()))
 			{
@@ -528,6 +702,7 @@ void UGameplayDebuggingComponent::CollectPathData()
 		}
 		else
 		{
+			NextPathPointIndex = INDEX_NONE;
 			CurrentPath = NULL;
 			PathPoints.Reset();
 			PathCorridorPolygons.Reset();
diff --git a/Engine/Source/Developer/GameplayDebugger/Private/GameplayDebuggingHUDComponent.cpp b/Engine/Source/Developer/GameplayDebugger/Private/GameplayDebuggingHUDComponent.cpp
--- a/Engine/Source/Developer/GameplayDebugger/Private/GameplayDebuggingHUDComponent.cpp
+++ b/Engine/Source/Developer/GameplayDebugger/Private/GameplayDebuggingHUDComponent.cpp
@@ -251,17 +251,19 @@ void AGameplayDebuggingHUDComponent::DrawPath(APlayerController* MyPC, class UGa
 {
 #if !(UE_BUILD_SHIPPING || UE_BUILD_TEST)
 	static const FColor Grey(100,100,100);
+	static const FColor PathColor(0, 255, 255);
+
 	const int32 NumPathVerts = DebugComponent->PathPoints.Num();
 	UWorld* World = GetWorld();
 
 		for (int32 VertIdx=0; VertIdx < NumPathVerts-1; ++VertIdx)
 		{
 			FVector const VertLoc = DebugComponent->PathPoints[VertIdx] + NavigationDebugDrawing::PathOffset;
-			DrawDebugSolidBox(World, VertLoc, NavigationDebugDrawing::PathNodeBoxExtent, Grey, false);
+			DrawDebugSolidBox(World, VertLoc, NavigationDebugDrawing::PathNodeBoxExtent, VertIdx < int32(DebugComponent->NextPathPointIndex) ? Grey : PathColor, false);
 
 			// draw line to next loc
 			FVector const NextVertLoc = DebugComponent->PathPoints[VertIdx+1] + NavigationDebugDrawing::PathOffset;
-			DrawDebugLine(World, VertLoc, NextVertLoc, Grey, false
+			DrawDebugLine(World, VertLoc, NextVertLoc, VertIdx < int32(DebugComponent->NextPathPointIndex) ? Grey : PathColor, false
 				, -1.f, 0
 				, NavigationDebugDrawing::PathLineThickness);
 		}
@@ -329,7 +331,6 @@ void AGameplayDebuggingHUDComponent::DrawOverHeadInformation(APlayerController*
 
 	if (EngineShowFlags.DebugAI)
 	{
-		PrintString(OverHeadContext, FString::Printf(TEXT("{red}%s\n"), *DebugComponent->PathErrorString));
 		DrawPath(PC, DebugComponent);
 	}
 #endif //!(UE_BUILD_SHIPPING || UE_BUILD_TEST)
@@ -345,10 +346,34 @@ void AGameplayDebuggingHUDComponent::DrawBasicData(APlayerController* PC, class
 	PrintString(DefaultContext, FString::Printf(TEXT("Controller Name: {yellow}%s\n"), *DebugComponent->ControllerName));
 	DefaultContext.Font = OldFont;
 
-	PrintString(DefaultContext, FString::Printf(TEXT("Pawn Name: {yellow}%s, {white}Pawn Class: {yellow}%s\n"), *DebugComponent->PawnName, *DebugComponent->PawnClass));
-	if (DebugComponent->PathErrorString.Len() > 0)
+	PrintString(DefaultContext, FString::Printf(TEXT("Pawn Name: {yellow}%s{white}, Pawn Class: {yellow}%s\n"), *DebugComponent->PawnName, *DebugComponent->PawnClass));
+
+	// movement
+	if (DebugComponent->bIsUsingCharacter)
+	{
+		PrintString(DefaultContext, FString::Printf(TEXT("Movement Mode: {yellow}%s{white}, Base: {yellow}%s\n"), *DebugComponent->MovementModeInfo, *DebugComponent->MovementBaseInfo));
+		PrintString(DefaultContext, FString::Printf(TEXT("NavData: {yellow}%s{white}, Path following: {yellow}%s\n"), *DebugComponent->NavDataInfo, *DebugComponent->PathFollowingInfo));
+	}
+
+	// logic
+	if (DebugComponent->bIsUsingBehaviorTree)
+	{
+		PrintString(DefaultContext, FString::Printf(TEXT("Behavior: {yellow}%s{white}, Tree: {yellow}%s\n"), *DebugComponent->CurrentAIState, *DebugComponent->CurrentAIAssets));
+		PrintString(DefaultContext, FString::Printf(TEXT("Active task: {yellow}%s\n"), *DebugComponent->CurrentAITask));
+	}
+
+	// ability + animation
+	if (DebugComponent->bIsUsingAbilities && DebugComponent->bIsUsingCharacter)
+	{
+		PrintString(DefaultContext, FString::Printf(TEXT("Ability: {yellow}%s{white}, Montage: {yellow}%s\n"), *DebugComponent->AbilityInfo, *DebugComponent->MontageInfo));
+	}
+	else if (DebugComponent->bIsUsingCharacter)
+	{
+		PrintString(DefaultContext, FString::Printf(TEXT("Montage: {yellow}%s\n"), *DebugComponent->MontageInfo));
+	}
+	else if (DebugComponent->bIsUsingAbilities)
 	{
-		PrintString(DefaultContext, FString::Printf(TEXT("{red}%s\n"), *DebugComponent->PathErrorString));
+		PrintString(DefaultContext, FString::Printf(TEXT("Ability: {yellow}%s\n"), *DebugComponent->AbilityInfo));
 	}
 
 	DrawPath(PC, DebugComponent);
