more visible path drawing color in gameplay debugger

diff --git a/Engine/Source/Developer/GameplayDebugger/Private/GameplayDebuggingHUDComponent.cpp b/Engine/Source/Developer/GameplayDebugger/Private/GameplayDebuggingHUDComponent.cpp
--- a/Engine/Source/Developer/GameplayDebugger/Private/GameplayDebuggingHUDComponent.cpp
+++ b/Engine/Source/Developer/GameplayDebugger/Private/GameplayDebuggingHUDComponent.cpp
@@ -251,7 +251,7 @@ void AGameplayDebuggingHUDComponent::DrawPath(APlayerController* MyPC, class UGa
 {
 #if !(UE_BUILD_SHIPPING || UE_BUILD_TEST)
 	static const FColor Grey(100,100,100);
-	static const FColor PathColor(0, 255, 255);
+	static const FColor PathColor(192,192,192);
 
 	const int32 NumPathVerts = DebugComponent->PathPoints.Num();
 	UWorld* World = GetWorld();
