#UE4 One more CIS fix. GameplayDebuggingComponent uses UCharacterMovementComponent

diff --git a/Engine/Source/Developer/GameplayDebugger/Private/GameplayDebuggingComponent.cpp b/Engine/Source/Developer/GameplayDebugger/Private/GameplayDebuggingComponent.cpp
--- a/Engine/Source/Developer/GameplayDebugger/Private/GameplayDebuggingComponent.cpp
+++ b/Engine/Source/Developer/GameplayDebugger/Private/GameplayDebuggingComponent.cpp
@@ -17,6 +17,7 @@
 #include "BehaviorTreeDelegates.h"
 #include "GameFramework/PlayerState.h"
 #include "GameFramework/Character.h"
+#include "GameFramework/CharacterMovementComponent.h"
 #include "Engine/Channel.h"
 
 #if WITH_EDITOR
