#UE4 One more CIS fix. Header includes shifted.

diff --git a/Engine/Source/Developer/GameplayDebugger/Private/GameplayDebuggingComponent.cpp b/Engine/Source/Developer/GameplayDebugger/Private/GameplayDebuggingComponent.cpp
--- a/Engine/Source/Developer/GameplayDebugger/Private/GameplayDebuggingComponent.cpp
+++ b/Engine/Source/Developer/GameplayDebugger/Private/GameplayDebuggingComponent.cpp
@@ -19,6 +19,7 @@
 #include "GameFramework/Character.h"
 #include "GameFramework/CharacterMovementComponent.h"
 #include "Engine/Channel.h"
+#include "Animation/AnimMontage.h"
 
 #if WITH_EDITOR
 #include "Editor/EditorEngine.h"
