Fixed NavMesh rendering with GameplayDebugger (works in client-server enviroment too)
Integration from Dev to Main

diff --git a/Engine/Source/Developer/GameplayDebugger/Private/GameplayDebugger.cpp b/Engine/Source/Developer/GameplayDebugger/Private/GameplayDebugger.cpp
--- a/Engine/Source/Developer/GameplayDebugger/Private/GameplayDebugger.cpp
+++ b/Engine/Source/Developer/GameplayDebugger/Private/GameplayDebugger.cpp
@@ -293,7 +293,7 @@ bool FGameplayDebugger::CreateGameplayDebuggerForPlayerController(APlayerControl
 	FActorSpawnParameters SpawnInfo;
 	SpawnInfo.bNoCollisionFail = true;
 	SpawnInfo.Name = *FString::Printf(TEXT("GameplayDebuggingReplicator_%s"), *PlayerController->GetName());
-	AGameplayDebuggingReplicator* DestActor = World->SpawnActor<AGameplayDebuggingReplicator>(SpawnInfo);
+	AGameplayDebuggingReplicator* DestActor = World->SpawnActor<AGameplayDebuggingReplicator>(FVector::ZeroVector, FRotator::ZeroRotator, SpawnInfo);
 	if (DestActor != NULL)
 	{
 		DestActor->SetLocalPlayerOwner(PlayerController);
@@ -357,7 +357,7 @@ void FGameplayDebugger::WorldAdded(UWorld* InWorld)
 	FActorSpawnParameters SpawnInfo;
 	SpawnInfo.bNoCollisionFail = true;
 	SpawnInfo.Name = *FString::Printf(TEXT("GameplayDebuggingReplicator_Global"));
-	AGameplayDebuggingReplicator* DestActor = InWorld->SpawnActor<AGameplayDebuggingReplicator>(SpawnInfo);
+	AGameplayDebuggingReplicator* DestActor = InWorld->SpawnActor<AGameplayDebuggingReplicator>(FVector::ZeroVector, FRotator::ZeroRotator, SpawnInfo);
 	if (DestActor != NULL)
 	{
 		DestActor->SetLocalPlayerOwner(NULL);
diff --git a/Engine/Source/Developer/GameplayDebugger/Private/GameplayDebuggingComponent.cpp b/Engine/Source/Developer/GameplayDebugger/Private/GameplayDebuggingComponent.cpp
--- a/Engine/Source/Developer/GameplayDebugger/Private/GameplayDebuggingComponent.cpp
+++ b/Engine/Source/Developer/GameplayDebugger/Private/GameplayDebuggingComponent.cpp
@@ -1282,6 +1282,7 @@ void UGameplayDebuggingComponent::PrepareNavMeshData(struct FNavMeshSceneProxyDa
 		int32 NumTiles = 0;
 		ArReader << NumTiles;
 
+		int32 IndexesOffset = 0;
 		for (int32 iTile = 0; iTile < NumTiles; iTile++)
 		{
 			NavMeshDebug::FTileData TileData;
@@ -1296,31 +1297,36 @@ void UGameplayDebuggingComponent::PrepareNavMeshData(struct FNavMeshSceneProxyDa
 				Verts.Add(Loc);
 			}
 			CurrentData->Bounds += FBox(Verts);
-
+			
 			for (int32 iArea = 0; iArea < TileData.Areas.Num(); iArea++)
 			{
 				const NavMeshDebug::FAreaPolys& SrcArea = TileData.Areas[iArea];
 				FNavMeshSceneProxyData::FDebugMeshData DebugMeshData;
 				DebugMeshData.ClusterColor = SrcArea.Color;
-
-				for (int32 iVert = 0; iVert < Verts.Num(); iVert++)
-				{
-					AddVertexHelper(DebugMeshData, Verts[iVert] + CurrentData->NavMeshDrawOffset);
-				}
+				DebugMeshData.ClusterColor.A = 128;
 				
 				for (int32 iTri = 0; iTri < SrcArea.Indices.Num(); iTri += 3)
 				{
-					AddTriangleHelper(DebugMeshData, SrcArea.Indices[iTri], SrcArea.Indices[iTri + 1], SrcArea.Indices[iTri + 2]);
+					const int32 Index0 = SrcArea.Indices[iTri + 0];
+					const int32 Index1 = SrcArea.Indices[iTri + 1];
+					const int32 Index2 = SrcArea.Indices[iTri + 2];
 
-					FVector V0 = Verts[SrcArea.Indices[iTri+0]];
-					FVector V1 = Verts[SrcArea.Indices[iTri+1]];
-					FVector V2 = Verts[SrcArea.Indices[iTri+2]];
+					FVector V0 = Verts[Index0];
+					FVector V1 = Verts[Index1];
+					FVector V2 = Verts[Index2];
 					CurrentData->TileEdgeLines.Add(FDebugRenderSceneProxy::FDebugLine(V0 + CurrentData->NavMeshDrawOffset, V1 + CurrentData->NavMeshDrawOffset, NavMeshRenderColor_Recast_TileEdges));
 					CurrentData->TileEdgeLines.Add(FDebugRenderSceneProxy::FDebugLine(V1 + CurrentData->NavMeshDrawOffset, V2 + CurrentData->NavMeshDrawOffset, NavMeshRenderColor_Recast_TileEdges));
 					CurrentData->TileEdgeLines.Add(FDebugRenderSceneProxy::FDebugLine(V2 + CurrentData->NavMeshDrawOffset, V0 + CurrentData->NavMeshDrawOffset, NavMeshRenderColor_Recast_TileEdges));
+
+					AddTriangleHelper(DebugMeshData, Index0 + IndexesOffset, Index1 + IndexesOffset, Index2 + IndexesOffset);
 				}
 
+				for (int32 iVert = 0; iVert < Verts.Num(); iVert++)
+				{
+					AddVertexHelper(DebugMeshData, Verts[iVert] + CurrentData->NavMeshDrawOffset);
+				}
 				CurrentData->MeshBuilders.Add(DebugMeshData);
+				IndexesOffset += Verts.Num();
 			}
 
 			for (int32 i = 0; i < TileData.Links.Num(); i++)
diff --git a/Engine/Source/Runtime/Engine/Private/NavMeshRenderingHelpers.h b/Engine/Source/Runtime/Engine/Private/NavMeshRenderingHelpers.h
--- a/Engine/Source/Runtime/Engine/Private/NavMeshRenderingHelpers.h
+++ b/Engine/Source/Runtime/Engine/Private/NavMeshRenderingHelpers.h
@@ -410,7 +410,7 @@ public:
 						FMeshBatch& Mesh = Collector.AllocateMesh();
 						FMeshBatchElement& BatchElement = Mesh.Elements[0];
 						BatchElement = MeshBatchElements[Index];
-						BatchElement.PrimitiveUniformBuffer = CreatePrimitiveUniformBufferImmediate(GetLocalToWorld(), GetBounds(), GetLocalBounds(), false, UseEditorDepthTest());
+						BatchElement.PrimitiveUniformBuffer = CreatePrimitiveUniformBufferImmediate(FMatrix::Identity, GetBounds(), GetLocalBounds(), false, UseEditorDepthTest());
 
 						Mesh.bWireframe = false;
 						Mesh.VertexFactory = &VertexFactory;
