Fix for UE-10042 "FIXIF: Crash starting a PIE session with no viewport"

diff --git a/Engine/Source/Developer/GameplayDebugger/Private/GameplayDebuggingReplicator.cpp b/Engine/Source/Developer/GameplayDebugger/Private/GameplayDebuggingReplicator.cpp
--- a/Engine/Source/Developer/GameplayDebugger/Private/GameplayDebuggingReplicator.cpp
+++ b/Engine/Source/Developer/GameplayDebugger/Private/GameplayDebuggingReplicator.cpp
@@ -198,7 +198,7 @@ void AGameplayDebuggingReplicator::BeginPlay()
 
 #if WITH_EDITOR
 	const UEditorEngine* EEngine = Cast<UEditorEngine>(GEngine);
-	if (EEngine && (EEngine->bIsSimulatingInEditor || EEngine->EditorWorld) && GetWorld() != EEngine->EditorWorld && !IsGlobalInWorld() && GCurrentLevelEditingViewportClient->EngineShowFlags.DebugAI)
+	if (EEngine && (EEngine->bIsSimulatingInEditor || EEngine->EditorWorld) && GetWorld() != EEngine->EditorWorld && !IsGlobalInWorld() && GCurrentLevelEditingViewportClient && GCurrentLevelEditingViewportClient->EngineShowFlags.DebugAI)
 	{
 		SetIsTemporarilyHiddenInEditor(false);
 		SetActorHiddenInGame(false);
diff --git a/Engine/Source/Runtime/AIModule/Private/BehaviorTree/BlackboardComponent.cpp b/Engine/Source/Runtime/AIModule/Private/BehaviorTree/BlackboardComponent.cpp
--- a/Engine/Source/Runtime/AIModule/Private/BehaviorTree/BlackboardComponent.cpp
+++ b/Engine/Source/Runtime/AIModule/Private/BehaviorTree/BlackboardComponent.cpp
@@ -472,7 +472,7 @@ FString UBlackboardComponent::DescribeKeyValue(FBlackboard::FKey KeyID, EBlackbo
 	if (Key)
 	{
 		const uint8* ValueData = GetKeyRawData(KeyID);
-		FString ValueDesc = Key->KeyType ? *(Key->KeyType->WrappedDescribeValue(*this, ValueData)) : TEXT("empty");
+		FString ValueDesc = Key->KeyType && ValueData ? *(Key->KeyType->WrappedDescribeValue(*this, ValueData)) : TEXT("empty");
 
 		if (Mode == EBlackboardDescription::OnlyValue)
 		{
