Fix for UE-9995 "GameplayDebugger doesn't auto-focus on AI in server-client PIE"

diff --git a/Engine/Source/Developer/GameplayDebugger/Classes/GameplayDebuggingComponent.h b/Engine/Source/Developer/GameplayDebugger/Classes/GameplayDebuggingComponent.h
--- a/Engine/Source/Developer/GameplayDebugger/Classes/GameplayDebuggingComponent.h
+++ b/Engine/Source/Developer/GameplayDebugger/Classes/GameplayDebuggingComponent.h
@@ -188,8 +188,8 @@ class GAMEPLAYDEBUGGER_API UGameplayDebuggingComponent : public UPrimitiveCompon
 	UPROPERTY(Replicated)
 	AActor* TargetActor;
 	
-	UFUNCTION(reliable, server, WithValidation)
-	void ServerEnableTargetSelection(bool bEnable);
+	UFUNCTION(Reliable, Client, WithValidation)
+	void ClientEnableTargetSelection(bool bEnable);
 
 	void SetActorToDebug(AActor* Actor);
 	FORCEINLINE AActor* GetSelectedActor() const
diff --git a/Engine/Source/Developer/GameplayDebugger/Classes/GameplayDebuggingReplicator.h b/Engine/Source/Developer/GameplayDebugger/Classes/GameplayDebuggingReplicator.h
--- a/Engine/Source/Developer/GameplayDebugger/Classes/GameplayDebuggingReplicator.h
+++ b/Engine/Source/Developer/GameplayDebugger/Classes/GameplayDebuggingReplicator.h
@@ -86,8 +86,8 @@ class GAMEPLAYDEBUGGER_API AGameplayDebuggingReplicator : public AActor
 	UFUNCTION(reliable, client, WithValidation)
 	void ClientReplicateMessage(class AActor* Actor, uint32 InMessage, uint32 DataView = 0);
 
-	UFUNCTION(reliable, server, WithValidation)
-	void ServerEnableTargetSelection(bool bEnable, APlayerController* Context);
+	UFUNCTION(Reliable, Client, WithValidation)
+	void ClientEnableTargetSelection(bool bEnable, APlayerController* Context);
 
 #if WITH_EDITOR
 	void PostEditChangeProperty(FPropertyChangedEvent& PropertyChangedEvent) override;
@@ -120,7 +120,9 @@ class GAMEPLAYDEBUGGER_API AGameplayDebuggingReplicator : public AActor
 	APlayerController* GetLocalPlayerOwner() { return LocalPlayerOwner; }
 
 	FORCEINLINE AActor* GetSelectedActorToDebug() { return LastSelectedActorToDebug; }
-	void SetActorToDebug(AActor* InActor);
+	
+	UFUNCTION(reliable, server, WithValidation)
+	void ServerSetActorToDebug(AActor* InActor);
 
 	/**
 	 * Iterates through the pawn list to find the next pawn of the specified type to debug
diff --git a/Engine/Source/Developer/GameplayDebugger/Private/GameplayDebuggingComponent.cpp b/Engine/Source/Developer/GameplayDebugger/Private/GameplayDebuggingComponent.cpp
--- a/Engine/Source/Developer/GameplayDebugger/Private/GameplayDebuggingComponent.cpp
+++ b/Engine/Source/Developer/GameplayDebugger/Private/GameplayDebuggingComponent.cpp
@@ -198,37 +198,24 @@ void UGameplayDebuggingComponent::GetLifetimeReplicatedProps( TArray< FLifetimeP
 	DOREPLIFETIME(UGameplayDebuggingComponent, PawnClass);
 
 	DOREPLIFETIME(UGameplayDebuggingComponent, bIsUsingCharacter);
-	DOREPLIFETIME(UGameplayDebuggingComponent, bIsUsingPathFollowing);
-	DOREPLIFETIME(UGameplayDebuggingComponent, bIsUsingBehaviorTree);
-	DOREPLIFETIME(UGameplayDebuggingComponent, bIsUsingAbilities);
-
-	if (bIsUsingCharacter)
-	{
-		DOREPLIFETIME(UGameplayDebuggingComponent, MovementBaseInfo);
-		DOREPLIFETIME(UGameplayDebuggingComponent, MovementModeInfo);
-		DOREPLIFETIME(UGameplayDebuggingComponent, MontageInfo);
-	}
+	DOREPLIFETIME(UGameplayDebuggingComponent, MovementBaseInfo);
+	DOREPLIFETIME(UGameplayDebuggingComponent, MovementModeInfo);
+	DOREPLIFETIME(UGameplayDebuggingComponent, MontageInfo);
 
-	if (bIsUsingPathFollowing)
-	{
-		DOREPLIFETIME(UGameplayDebuggingComponent, PathFollowingInfo);
-		DOREPLIFETIME(UGameplayDebuggingComponent, NavDataInfo);
-		DOREPLIFETIME(UGameplayDebuggingComponent, PathPoints);
-		DOREPLIFETIME(UGameplayDebuggingComponent, PathCorridorData);
-		DOREPLIFETIME(UGameplayDebuggingComponent, NextPathPointIndex);
-	}
+	DOREPLIFETIME(UGameplayDebuggingComponent, bIsUsingPathFollowing);
+	DOREPLIFETIME(UGameplayDebuggingComponent, PathFollowingInfo);
+	DOREPLIFETIME(UGameplayDebuggingComponent, NavDataInfo);
+	DOREPLIFETIME(UGameplayDebuggingComponent, PathPoints);
+	DOREPLIFETIME(UGameplayDebuggingComponent, PathCorridorData);
+	DOREPLIFETIME(UGameplayDebuggingComponent, NextPathPointIndex);
 
-	if (bIsUsingBehaviorTree)
-	{
-		DOREPLIFETIME(UGameplayDebuggingComponent, CurrentAITask);
-		DOREPLIFETIME(UGameplayDebuggingComponent, CurrentAIState);
-		DOREPLIFETIME(UGameplayDebuggingComponent, CurrentAIAssets);
-	}
+	DOREPLIFETIME(UGameplayDebuggingComponent, bIsUsingBehaviorTree);
+	DOREPLIFETIME(UGameplayDebuggingComponent, CurrentAITask);
+	DOREPLIFETIME(UGameplayDebuggingComponent, CurrentAIState);
+	DOREPLIFETIME(UGameplayDebuggingComponent, CurrentAIAssets);
 
-	if (bIsUsingAbilities)
-	{
-		DOREPLIFETIME(UGameplayDebuggingComponent, AbilityInfo);
-	}
+	DOREPLIFETIME(UGameplayDebuggingComponent, bIsUsingAbilities);
+	DOREPLIFETIME(UGameplayDebuggingComponent, AbilityInfo);
 
 	DOREPLIFETIME(UGameplayDebuggingComponent, BrainComponentName);
 	DOREPLIFETIME(UGameplayDebuggingComponent, BrainComponentString);
@@ -242,16 +229,16 @@ void UGameplayDebuggingComponent::GetLifetimeReplicatedProps( TArray< FLifetimeP
 #endif //!(UE_BUILD_SHIPPING || UE_BUILD_TEST)
 }
 
-bool UGameplayDebuggingComponent::ServerEnableTargetSelection_Validate(bool bEnable)
+bool UGameplayDebuggingComponent::ClientEnableTargetSelection_Validate(bool bEnable)
 {
 	return true;
 }
 
-void UGameplayDebuggingComponent::ServerEnableTargetSelection_Implementation(bool bEnable)
+void UGameplayDebuggingComponent::ClientEnableTargetSelection_Implementation(bool bEnable)
 {
 #if !(UE_BUILD_SHIPPING || UE_BUILD_TEST)
 	bEnabledTargetSelection = bEnable;
-	if (bEnabledTargetSelection && World && World->GetNetMode() < NM_Client)
+	if (bEnabledTargetSelection && World && World->GetNetMode() != NM_DedicatedServer)
 	{
 		NextTargrtSelectionTime = 0;
 	}
@@ -271,7 +258,7 @@ void UGameplayDebuggingComponent::TickComponent(float DeltaTime, enum ELevelTick
 	Super::TickComponent(DeltaTime, TickType, ThisTickFunction);
 #if !(UE_BUILD_SHIPPING || UE_BUILD_TEST)
 
-	if (World && World->GetNetMode() < NM_Client)
+	if (World && World->GetNetMode() != NM_DedicatedServer)
 	{
 		if (bEnabledTargetSelection)
 		{
@@ -284,10 +271,12 @@ void UGameplayDebuggingComponent::TickComponent(float DeltaTime, enum ELevelTick
 				}
 			}
 		}
-		CollectDataToReplicate(true);
 	}
 
-	AGameplayDebuggingReplicator* Replicator = Cast<AGameplayDebuggingReplicator>(GetOwner());
+	if (World && World->GetNetMode() < NM_Client)
+	{
+		CollectDataToReplicate(true);
+	}
 
 #endif //!(UE_BUILD_SHIPPING || UE_BUILD_TEST)
 }
@@ -300,17 +289,7 @@ void UGameplayDebuggingComponent::SelectTargetToDebug()
 
 	if (MyPC )
 	{
-		APawn* BestTarget = NULL;
-		if (MyPC->GetViewTarget() != NULL && MyPC->GetViewTarget() != MyPC->GetPawn())
-		{
-			BestTarget = Cast<APawn>(MyPC->GetViewTarget());
-			if (BestTarget && ((BestTarget->PlayerState != NULL && BestTarget->PlayerState->bIsABot == false) || BestTarget->GetActorEnableCollision() == false))
-			{
-				BestTarget = NULL;
-			}
-		}
-
-		float bestAim = 0.f;
+		float bestAim = 0;
 		FVector CamLocation;
 		FRotator CamRotation;
 		check(MyPC->PlayerCameraManager != NULL);
@@ -330,7 +309,7 @@ void UGameplayDebuggingComponent::SelectTargetToDebug()
 				continue;
 			}
 			
-			if (BestTarget == NULL && (NewTarget != MyPC->GetPawn()))
+			if (NewTarget != MyPC->GetPawn())
 			{
 				// look for best controlled pawn target
 				const FVector AimDir = NewTarget->GetActorLocation() - CamLocation;
@@ -339,8 +318,8 @@ void UGameplayDebuggingComponent::SelectTargetToDebug()
 				if (FireDist < 625000000.f)
 				{
 					FireDist = FMath::Sqrt(FireDist);
-					float newAim = FireDir | AimDir;
-					newAim = newAim/FireDist;
+					float newAim = FVector::DotProduct(FireDir, AimDir);
+					newAim = newAim / FireDist;
 					if (newAim > bestAim)
 					{
 						PossibleTarget = NewTarget;
@@ -350,16 +329,13 @@ void UGameplayDebuggingComponent::SelectTargetToDebug()
 			}
 		}
 
-		BestTarget = BestTarget == NULL ? PossibleTarget : BestTarget;
-		if (BestTarget != NULL && BestTarget != GetSelectedActor())
+		if (PossibleTarget != NULL && PossibleTarget != GetSelectedActor())
 		{
 			if (AGameplayDebuggingReplicator* Replicator = Cast<AGameplayDebuggingReplicator>(GetOwner()))
 			{
-				Replicator->SetActorToDebug(Cast<AActor>(BestTarget));
+				Replicator->ServerSetActorToDebug(Cast<AActor>(PossibleTarget));
 			}
 
-			//always update component for best target
-			SetActorToDebug(Cast<AActor>(BestTarget));
 			ServerReplicateData(EDebugComponentMessage::ActivateReplication, EAIDebugDrawDataView::Empty);
 		}
 	}
@@ -421,9 +397,6 @@ void UGameplayDebuggingComponent::CollectBasicData()
 	AAIController* MyController = Cast<AAIController>(MyPawn->Controller);
 
 	bIsUsingCharacter = MyPawn->IsA(ACharacter::StaticClass());
-	bIsUsingPathFollowing = false;
-	bIsUsingBehaviorTree = false;
-	bIsUsingAbilities = false;
 
 	if (MyController)
 	{
@@ -703,6 +676,7 @@ void UGameplayDebuggingComponent::CollectPathData()
 							ArWriter << Vert;
 						}
 					}
+					PathCorridorData = HelpBuffer;
 				}
 				bRefreshRendering = true;
 			}
@@ -729,8 +703,7 @@ void UGameplayDebuggingComponent::OnRep_PathCorridorData()
 #if !(UE_BUILD_SHIPPING || UE_BUILD_TEST)
 	if (World && World->GetNetMode() != NM_DedicatedServer)
 	{
-		TArray<uint8> HelpBuffer;
-		FMemoryReader ArReader(HelpBuffer);
+		FMemoryReader ArReader(PathCorridorData);
 		int32 NumPolygons = 0;
 		ArReader << NumPolygons;
 		PathCorridorPolygons.Reset(NumPolygons);
@@ -745,6 +718,10 @@ void UGameplayDebuggingComponent::OnRep_PathCorridorData()
 				ArReader << CurrentVertex;
 				CurrentPoly.Add(CurrentVertex);
 			}
+			if (NumVerts > 2)
+			{
+				PathCorridorPolygons.Add(CurrentPoly);
+			}
 		}
 		UpdateBounds();
 		MarkRenderStateDirty();
diff --git a/Engine/Source/Developer/GameplayDebugger/Private/GameplayDebuggingControllerComponent.cpp b/Engine/Source/Developer/GameplayDebugger/Private/GameplayDebuggingControllerComponent.cpp
--- a/Engine/Source/Developer/GameplayDebugger/Private/GameplayDebuggingControllerComponent.cpp
+++ b/Engine/Source/Developer/GameplayDebugger/Private/GameplayDebuggingControllerComponent.cpp
@@ -66,7 +66,7 @@ void UGameplayDebuggingControllerComponent::BeginDestroy()
 
 	if (GetDebuggingReplicator() && PlayerOwner.IsValid() && !PlayerOwner->IsPendingKill())
 	{
-		APawn* Pawn = PlayerOwner->GetPawn();
+		APawn* Pawn = PlayerOwner->GetPawnOrSpectator();
 		if (Pawn && !Pawn->IsPendingKill())
 		{
 			GetDebuggingReplicator()->ServerReplicateMessage(Pawn, EDebugComponentMessage::DeactivateReplilcation, 0);
@@ -209,7 +209,7 @@ void UGameplayDebuggingControllerComponent::CloseDebugTool()
 
 void UGameplayDebuggingControllerComponent::EnableTargetSelection(bool bEnable)
 {
-	GetDebuggingReplicator()->ServerEnableTargetSelection(bEnable, PlayerOwner.Get());
+	GetDebuggingReplicator()->ClientEnableTargetSelection(bEnable, PlayerOwner.Get());
 }
 
 void UGameplayDebuggingControllerComponent::BindAIDebugViewKeys()
@@ -265,7 +265,7 @@ bool UGameplayDebuggingControllerComponent::CanToggleView()
 void UGameplayDebuggingControllerComponent::ToggleAIDebugView_SetView0()
 {
 #if !(UE_BUILD_SHIPPING || UE_BUILD_TEST)
-	APawn* Pawn = PlayerOwner.IsValid() ? PlayerOwner->GetPawn() : NULL;
+	APawn* Pawn = PlayerOwner.IsValid() ? PlayerOwner->GetPawnOrSpectator() : NULL;
 
 	if (PlayerOwner.IsValid() && Pawn && bToolActivated && GetDebuggingReplicator())
 	{
@@ -353,7 +353,7 @@ float UGameplayDebuggingControllerComponent::GetUpdateNavMeshTimeRemaining() con
 
 void UGameplayDebuggingControllerComponent::UpdateNavMeshTimer()
 {
-	const APawn* PlayerPawn = PlayerOwner.IsValid() ? PlayerOwner->GetPawn() : NULL;
+	const APawn* PlayerPawn = PlayerOwner.IsValid() ? PlayerOwner->GetPawnOrSpectator() : NULL;
 	UGameplayDebuggingComponent* DebuggingComponent = GetDebuggingReplicator() ? GetDebuggingReplicator()->GetDebugComponent() : NULL;
 	if (DebuggingComponent)
 	{
diff --git a/Engine/Source/Developer/GameplayDebugger/Private/GameplayDebuggingReplicator.cpp b/Engine/Source/Developer/GameplayDebugger/Private/GameplayDebuggingReplicator.cpp
--- a/Engine/Source/Developer/GameplayDebugger/Private/GameplayDebuggingReplicator.cpp
+++ b/Engine/Source/Developer/GameplayDebugger/Private/GameplayDebuggingReplicator.cpp
@@ -82,7 +82,7 @@ AGameplayDebuggingReplicator::AGameplayDebuggingReplicator(const FObjectInitiali
 		SetRemoteRoleForBackwardsCompat(ROLE_SimulatedProxy);
 		SetReplicates(true);
 
-		AGameplayDebuggingReplicator::OnSelectionChangedDelegate.AddUObject(this, &AGameplayDebuggingReplicator::SetActorToDebug);
+		AGameplayDebuggingReplicator::OnSelectionChangedDelegate.AddUObject(this, &AGameplayDebuggingReplicator::ServerSetActorToDebug);
 	}
 }
 
@@ -292,16 +292,16 @@ class UNetConnection* AGameplayDebuggingReplicator::GetNetConnection()
 	return NULL;
 }
 
-bool AGameplayDebuggingReplicator::ServerEnableTargetSelection_Validate(bool, APlayerController* )
+bool AGameplayDebuggingReplicator::ClientEnableTargetSelection_Validate(bool, APlayerController* )
 {
 	return true;
 }
 
-void AGameplayDebuggingReplicator::ServerEnableTargetSelection_Implementation(bool bEnable, APlayerController* Context)
+void AGameplayDebuggingReplicator::ClientEnableTargetSelection_Implementation(bool bEnable, APlayerController* Context)
 {
 	if (GetDebugComponent())
 	{
-		GetDebugComponent()->ServerEnableTargetSelection(bEnable);
+		GetDebugComponent()->ClientEnableTargetSelection(bEnable);
 	}
 }
 
@@ -325,7 +325,7 @@ void AGameplayDebuggingReplicator::ServerReplicateMessage_Implementation(class
 #if !(UE_BUILD_SHIPPING || UE_BUILD_TEST)
 	if ((EDebugComponentMessage::Type)InMessage == EDebugComponentMessage::DeactivateReplilcation)
 	{
-		SetActorToDebug(NULL);
+		ServerSetActorToDebug(NULL);
 		MarkComponentsRenderStateDirty();
 	}
 
@@ -431,7 +431,12 @@ void AGameplayDebuggingReplicator::TickActor(float DeltaTime, enum ELevelTick Ti
 	}
 }
 
-void AGameplayDebuggingReplicator::SetActorToDebug(AActor* InActor) 
+bool AGameplayDebuggingReplicator::ServerSetActorToDebug_Validate(AActor* InActor)
+{
+	return true;
+}
+
+void AGameplayDebuggingReplicator::ServerSetActorToDebug_Implementation(AActor* InActor)
 { 
 	if (LastSelectedActorToDebug != InActor)
 	{
@@ -516,7 +521,7 @@ void AGameplayDebuggingReplicator::OnDebugAIDelegate(class UCanvas* Canvas, clas
 		}
 
 		const AActor* OldActor = LastSelectedActorToDebug;
-		SetActorToDebug(FullSelectedTarget);
+		ServerSetActorToDebug(FullSelectedTarget);
 		if (FullSelectedTarget)
 		{
 			GetDebugComponent()->CollectDataToReplicate(true);
@@ -575,7 +580,7 @@ void AGameplayDebuggingReplicator::DrawDebugDataDelegate(class UCanvas* Canvas,
 			AActor* NewTarget = Cast<AActor>(*Iterator);
 			if (NewTarget->IsSelected() && GetSelectedActorToDebug() != NewTarget)
 			{
-				SetActorToDebug(NewTarget);
+				ServerSetActorToDebug(NewTarget);
 			}
 
 			GetDebugComponent()->SetActorToDebug(NewTarget);
@@ -633,7 +638,7 @@ void AGameplayDebuggingReplicator::DebugNextPawn(UClass* CompareClass, APawn* Cu
 		{
 			if (LastSeen == CurrentPawn)
 			{
-				SetActorToDebug(IterPawn);
+				ServerSetActorToDebug(IterPawn);
 				return;
 			}
 			LastSeen = IterPawn;
@@ -646,7 +651,7 @@ void AGameplayDebuggingReplicator::DebugNextPawn(UClass* CompareClass, APawn* Cu
 	// See if we need to wrap around the list
 	if (FirstSeen != nullptr)
 	{
-		SetActorToDebug(FirstSeen);
+		ServerSetActorToDebug(FirstSeen);
 	}
 }
 
@@ -663,7 +668,7 @@ void AGameplayDebuggingReplicator::DebugPrevPawn(UClass* CompareClass, APawn* Cu
 		{
 			if (LastSeen == CurrentPawn && FirstSeen != CurrentPawn)
 			{
-				SetActorToDebug(PrevSeen);
+				ServerSetActorToDebug(PrevSeen);
 				return;
 			}
 			PrevSeen = LastSeen;
@@ -673,7 +678,7 @@ void AGameplayDebuggingReplicator::DebugPrevPawn(UClass* CompareClass, APawn* Cu
 				FirstSeen = IterPawn;
 				if (CurrentPawn == nullptr)
 				{
-					SetActorToDebug(FirstSeen);
+					ServerSetActorToDebug(FirstSeen);
 					return;
 				}
 			}
@@ -682,11 +687,11 @@ void AGameplayDebuggingReplicator::DebugPrevPawn(UClass* CompareClass, APawn* Cu
 	// Wrap from the beginning to the end
 	if (FirstSeen == CurrentPawn)
 	{
-		SetActorToDebug(LastSeen);
+		ServerSetActorToDebug(LastSeen);
 	}
 	// Handle getting the previous to the end
 	else if (LastSeen == CurrentPawn)
 	{
-		SetActorToDebug(PrevSeen);
+		ServerSetActorToDebug(PrevSeen);
 	}
 }
