Improved path corridor drawing in Gameplay Debugger and we can see nav areas rendered with different colors - with colors defined in navarea.

diff --git a/Engine/Source/Developer/GameplayDebugger/Classes/GameplayDebuggingComponent.h b/Engine/Source/Developer/GameplayDebugger/Classes/GameplayDebuggingComponent.h
--- a/Engine/Source/Developer/GameplayDebugger/Classes/GameplayDebuggingComponent.h
+++ b/Engine/Source/Developer/GameplayDebugger/Classes/GameplayDebuggingComponent.h
@@ -40,6 +40,12 @@ class GAMEPLAYDEBUGGER_API UGameplayDebuggingComponent : public UPrimitiveCompon
 
 	friend class AGameplayDebuggingHUDComponent;
 
+	struct FPathCorridorPolygons
+	{
+		TArray<FVector> Points;
+		FColor Color;
+	};
+
 	UPROPERTY(globalconfig)
 	FString DebugComponentClassName;
 
@@ -106,7 +112,7 @@ class GAMEPLAYDEBUGGER_API UGameplayDebuggingComponent : public UPrimitiveCompon
 	UPROPERTY(ReplicatedUsing = OnRep_PathCorridorData)
 	TArray<uint8> PathCorridorData;
 	
-	TArray<TArray<FVector> > PathCorridorPolygons;
+	TArray<FPathCorridorPolygons> PathCorridorPolygons;
 	/** End path replication data*/
 	
 	UPROPERTY(ReplicatedUsing = OnRep_UpdateNavmesh)
diff --git a/Engine/Source/Developer/GameplayDebugger/Private/GameplayDebuggingComponent.cpp b/Engine/Source/Developer/GameplayDebugger/Private/GameplayDebuggingComponent.cpp
--- a/Engine/Source/Developer/GameplayDebugger/Private/GameplayDebuggingComponent.cpp
+++ b/Engine/Source/Developer/GameplayDebugger/Private/GameplayDebuggingComponent.cpp
@@ -645,11 +645,15 @@ void UGameplayDebuggingComponent::CollectPathData()
 				FVisualLogEntry Snapshot;
 				NewPath->DescribeSelfToVisLog(&Snapshot);
 				PathCorridorPolygons.Reset(Snapshot.ElementsToDraw.Num());
-				for (auto& CurrentShape : Snapshot.ElementsToDraw)
+				FPathCorridorPolygons Polygon;
+				for (FVisualLogShapeElement& CurrentShape : Snapshot.ElementsToDraw)
 				{
 					if (CurrentShape.GetType() == EVisualLoggerShapeElement::Polygon)
 					{
-						PathCorridorPolygons.Add(CurrentShape.Points);
+						Polygon.Color = CurrentShape.GetFColor();
+						Polygon.Points.Reset();
+						Polygon.Points.Append(CurrentShape.Points);
+						PathCorridorPolygons.Add(Polygon);
 					}
 				}
 
@@ -667,11 +671,13 @@ void UGameplayDebuggingComponent::CollectPathData()
 					FMemoryWriter ArWriter(HelpBuffer);
 					int32 NumPolygons = PathCorridorPolygons.Num();
 					ArWriter << NumPolygons;
-					for (const auto& CurrentPoly : PathCorridorPolygons)
+					for (const FPathCorridorPolygons& CurrentPoly : PathCorridorPolygons)
 					{
-						int32 NumVerts = CurrentPoly.Num();
+						FColor Color = CurrentPoly.Color;
+						ArWriter << Color;
+						int32 NumVerts = CurrentPoly.Points.Num();
 						ArWriter << NumVerts;
-						for (auto Vert : CurrentPoly)
+						for (auto Vert : CurrentPoly.Points)
 						{
 							ArWriter << Vert;
 						}
@@ -709,18 +715,21 @@ void UGameplayDebuggingComponent::OnRep_PathCorridorData()
 		PathCorridorPolygons.Reset(NumPolygons);
 		for (int32 PolyIndex = 0; PolyIndex < NumPolygons; ++PolyIndex)
 		{
-			TArray<FVector> CurrentPoly;
+			FPathCorridorPolygons Polygon;
+			FColor Color;
+			ArReader << Polygon.Color;
+
 			int32 NumVerts = 0;
 			ArReader << NumVerts;
 			for (int32 VertIndex = 0; VertIndex < NumVerts; ++VertIndex)
 			{
 				FVector CurrentVertex;
 				ArReader << CurrentVertex;
-				CurrentPoly.Add(CurrentVertex);
+				Polygon.Points.Add(CurrentVertex);
 			}
 			if (NumVerts > 2)
 			{
-				PathCorridorPolygons.Add(CurrentPoly);
+				PathCorridorPolygons.Add(Polygon);
 			}
 		}
 		UpdateBounds();
@@ -1353,6 +1362,7 @@ public:
 		: FDebugRenderSceneProxy(InComponent)
 	{
 		DrawType = FDebugRenderSceneProxy::SolidAndWireMeshes;
+		DrawAlpha = 90;
 		ViewFlagName = InViewFlagName;
 		ViewFlagIndex = uint32(FEngineShowFlags::FindIndexByName(*ViewFlagName));
 	}
@@ -1427,44 +1437,33 @@ FPrimitiveSceneProxy* UGameplayDebuggingComponent::CreateSceneProxy()
 
 		TArray<FDebugRenderSceneProxy::FMesh> Meshes;
 		TArray<FDebugRenderSceneProxy::FDebugLine> Lines;
-		for (auto& CurrentPoly : PathCorridorPolygons)
+		for (FPathCorridorPolygons& CurrentPoly : PathCorridorPolygons)
 		{
-#if WITH_EDITOR
-			FClipSMPolygon InPoly(CurrentPoly.Num());
-			for (int32 Index = 0; Index < CurrentPoly.Num(); Index++)
-			{
-				FClipSMVertex v1;
-				v1.Pos = CurrentPoly[Index];
-				InPoly.Vertices.Add(v1);
-			}
 
-			TArray<FClipSMTriangle> OutTris;
-			if (TriangulatePoly(OutTris, InPoly, false))
+			if (CurrentPoly.Points.Num() > 2)
 			{
 				int32 LastIndex = 0;
+				FVector FirstVertex = CurrentPoly.Points[0];
 				FDebugRenderSceneProxy::FMesh TestMesh;
-				TestMesh.Color = FColor::Green;
-
-				RemoveRedundantTriangles(OutTris);
-				for (const auto& CurrentTri : OutTris)
+				TestMesh.Color = CurrentPoly.Color;// FColor::Green;
+				for (int32 Index = 1; Index < CurrentPoly.Points.Num()-1; Index += 1)
 				{
-					TestMesh.Vertices.Add(FDynamicMeshVertex(CurrentTri.Vertices[0].Pos));
-					TestMesh.Vertices.Add(FDynamicMeshVertex(CurrentTri.Vertices[1].Pos));
-					TestMesh.Vertices.Add(FDynamicMeshVertex(CurrentTri.Vertices[2].Pos));
+					TestMesh.Vertices.Add(FDynamicMeshVertex(FirstVertex));
+					TestMesh.Vertices.Add(FDynamicMeshVertex(CurrentPoly.Points[Index + 0]));
+					TestMesh.Vertices.Add(FDynamicMeshVertex(CurrentPoly.Points[Index + 1]));
 					TestMesh.Indices.Add(LastIndex++);
 					TestMesh.Indices.Add(LastIndex++);
 					TestMesh.Indices.Add(LastIndex++);
 				}
 				Meshes.Add(TestMesh);
 			}
-#endif
 
-			for (int32 VIdx = 0; VIdx < CurrentPoly.Num(); VIdx++)
+			for (int32 VIdx = 0; VIdx < CurrentPoly.Points.Num(); VIdx++)
 			{
 				Lines.Add(FDebugRenderSceneProxy::FDebugLine(
-					CurrentPoly[VIdx],
-					CurrentPoly[(VIdx + 1) % CurrentPoly.Num()],
-					FColor::Cyan,
+					CurrentPoly.Points[VIdx],
+					CurrentPoly.Points[(VIdx + 1) % CurrentPoly.Points.Num()],
+					CurrentPoly.Color,
 					2)
 				);
 			}
diff --git a/Engine/Source/Runtime/Engine/Private/AI/Navigation/NavigationPath.cpp b/Engine/Source/Runtime/Engine/Private/AI/Navigation/NavigationPath.cpp
--- a/Engine/Source/Runtime/Engine/Private/AI/Navigation/NavigationPath.cpp
+++ b/Engine/Source/Runtime/Engine/Private/AI/Navigation/NavigationPath.cpp
@@ -1084,14 +1084,20 @@ void FNavMeshPath::DescribeSelfToVisLog(FVisualLogEntry* Snapshot) const
 	TArray<FVector> Verts;
 	for (int32 Idx = 0; Idx < PathCorridor.Num(); Idx++)
 	{
+		const uint8 AreaID = NavMesh->GetPolyAreaID(PathCorridor[Idx]);
+		const UClass* AreaClass = NavMesh->GetAreaClass(AreaID);
+		
 		Verts.Reset();
 		NavMesh->GetPolyVerts(PathCorridor[Idx], Verts);
+
+		const UNavArea* DefArea = AreaClass ? ((UClass*)AreaClass)->GetDefaultObject<UNavArea>() : NULL;
+		const FColor PolygonColor = AreaClass != UNavigationSystem::GetDefaultWalkableArea() ? (DefArea ? DefArea->DrawColor : NavMesh->GetConfig().Color) : FColorList::Cyan;
+
+		CorridorPoly.SetColor(PolygonColor);
 		CorridorPoly.Points.Reset();
 		CorridorPoly.Points.Append(Verts);
 		Snapshot->ElementsToDraw.Add(CorridorPoly);
 
-		const uint8 AreaID = NavMesh->GetPolyAreaID(PathCorridor[Idx]);
-		const UClass* AreaClass = NavMesh->GetAreaClass(AreaID);
 		if (AreaClass && AreaClass != UNavigationSystem::GetDefaultWalkableArea())
 		{
 			FVector CenterPt = FVector::ZeroVector;
diff --git a/Engine/Source/Runtime/Engine/Private/DebugRenderSceneProxy.cpp b/Engine/Source/Runtime/Engine/Private/DebugRenderSceneProxy.cpp
--- a/Engine/Source/Runtime/Engine/Private/DebugRenderSceneProxy.cpp
+++ b/Engine/Source/Runtime/Engine/Private/DebugRenderSceneProxy.cpp
@@ -19,6 +19,7 @@ FDebugRenderSceneProxy::FDebugRenderSceneProxy(const UPrimitiveComponent* InComp
 	, ViewFlagName(TEXT("Game"))
 	, TextWithoutShadowDistance(1500)
 	, DrawType(WireMesh)
+	, DrawAlpha(60)
 {
 
 }
@@ -154,7 +155,7 @@ void FDebugRenderSceneProxy::GetDynamicMeshElements(const TArray<const FSceneVie
 
 				if (DrawType & SolidAndWireMeshes || DrawType & SolidMesh)
 				{
-					GetCylinderMesh(Cylinder.Base, FVector(1, 0, 0), FVector(0, 1, 0), FVector(0, 0, 1), Cylinder.Radius, Cylinder.HalfHeight, 16, MaterialCache[Cylinder.Color.WithAlpha(60)], SDPG_World, ViewIndex, Collector);
+					GetCylinderMesh(Cylinder.Base, FVector(1, 0, 0), FVector(0, 1, 0), FVector(0, 0, 1), Cylinder.Radius, Cylinder.HalfHeight, 16, MaterialCache[Cylinder.Color.WithAlpha(DrawAlpha)], SDPG_World, ViewIndex, Collector);
 				}
 			}
 
@@ -167,7 +168,7 @@ void FDebugRenderSceneProxy::GetDynamicMeshElements(const TArray<const FSceneVie
 				}
 				if (DrawType & SolidAndWireMeshes || DrawType & SolidMesh)
 				{
-					GetBoxMesh(FTransform(Box.Box.GetCenter()).ToMatrixNoScale() * Box.Transform.ToMatrixWithScale(), Box.Box.GetExtent(), MaterialCache[Box.Color.WithAlpha(60)], SDPG_World, ViewIndex, Collector);
+					GetBoxMesh(FTransform(Box.Box.GetCenter()).ToMatrixNoScale() * Box.Transform.ToMatrixWithScale(), Box.Box.GetExtent(), MaterialCache[Box.Color.WithAlpha(DrawAlpha)], SDPG_World, ViewIndex, Collector);
 				}
 			}
 
@@ -181,7 +182,7 @@ void FDebugRenderSceneProxy::GetDynamicMeshElements(const TArray<const FSceneVie
 				}
 				if (DrawType & SolidAndWireMeshes || DrawType & SolidMesh)
 				{
-					GetConeMesh(CurrentCone.ConeToWorld, CurrentCone.Angle1, CurrentCone.Angle2, 16, MaterialCache[CurrentCone.Color.WithAlpha(60)], SDPG_World, ViewIndex, Collector);
+					GetConeMesh(CurrentCone.ConeToWorld, CurrentCone.Angle1, CurrentCone.Angle2, 16, MaterialCache[CurrentCone.Color.WithAlpha(DrawAlpha)], SDPG_World, ViewIndex, Collector);
 				}
 			}
 
@@ -195,7 +196,7 @@ void FDebugRenderSceneProxy::GetDynamicMeshElements(const TArray<const FSceneVie
 					}
 					if (DrawType & SolidAndWireMeshes || DrawType & SolidMesh)
 					{
-						GetSphereMesh(It->Location, FVector(It->Radius), 20, 7, MaterialCache[It->Color.WithAlpha(60)], SDPG_World, false, ViewIndex, Collector);
+						GetSphereMesh(It->Location, FVector(It->Radius), 20, 7, MaterialCache[It->Color.WithAlpha(DrawAlpha)], SDPG_World, false, ViewIndex, Collector);
 					}
 				}
 			}
@@ -215,7 +216,7 @@ void FDebugRenderSceneProxy::GetDynamicMeshElements(const TArray<const FSceneVie
 					}
 					if (DrawType & SolidAndWireMeshes || DrawType & SolidMesh)
 					{
-						GetCapsuleMesh(It->Location, It->X, It->Y, It->Z, It->Color, It->Radius, It->HalfHeight, 16, MaterialCache[It->Color.WithAlpha(60)], SDPG_World, false, ViewIndex, Collector);
+						GetCapsuleMesh(It->Location, It->X, It->Y, It->Z, It->Color, It->Radius, It->HalfHeight, 16, MaterialCache[It->Color.WithAlpha(DrawAlpha)], SDPG_World, false, ViewIndex, Collector);
 					}
 				}
 			}
@@ -226,7 +227,7 @@ void FDebugRenderSceneProxy::GetDynamicMeshElements(const TArray<const FSceneVie
 				MeshBuilder.AddVertices(Mesh.Vertices);
 				MeshBuilder.AddTriangles(Mesh.Indices);
 
-				MeshBuilder.GetMesh(FMatrix::Identity, MaterialCache[Mesh.Color.WithAlpha(60)], SDPG_World, false, false, ViewIndex, Collector);
+				MeshBuilder.GetMesh(FMatrix::Identity, MaterialCache[Mesh.Color.WithAlpha(DrawAlpha)], SDPG_World, false, false, ViewIndex, Collector);
 			}
 
 		}
diff --git a/Engine/Source/Runtime/Engine/Public/DebugRenderSceneProxy.h b/Engine/Source/Runtime/Engine/Public/DebugRenderSceneProxy.h
--- a/Engine/Source/Runtime/Engine/Public/DebugRenderSceneProxy.h
+++ b/Engine/Source/Runtime/Engine/Public/DebugRenderSceneProxy.h
@@ -243,6 +243,7 @@ public:
 	FDebugDrawDelegate DebugTextDrawingDelegate;
 	FDelegateHandle DebugTextDrawingDelegateHandle;
 	EDrawType DrawType;
+	uint32 DrawAlpha;
 };
 
 #endif
