Make pass over gameplay debugger to be sure that most code is compile-out for test and shipping builds.

diff --git a/Engine/Source/Developer/GameplayDebugger/Private/GameplayDebugger.cpp b/Engine/Source/Developer/GameplayDebugger/Private/GameplayDebugger.cpp
--- a/Engine/Source/Developer/GameplayDebugger/Private/GameplayDebugger.cpp
+++ b/Engine/Source/Developer/GameplayDebugger/Private/GameplayDebugger.cpp
@@ -335,6 +335,7 @@ void FGameplayDebugger::RemoveReplicator(UWorld* InWorld, AGameplayDebuggingRepl
 
 void FGameplayDebugger::WorldAdded(UWorld* InWorld)
 {
+#if !(UE_BUILD_SHIPPING || UE_BUILD_TEST)
 	bool bIsServer = InWorld && InWorld->GetNetMode() < ENetMode::NM_Client; // (Only server code)
 	if (!bIsServer)
 	{
@@ -369,10 +370,12 @@ void FGameplayDebugger::WorldAdded(UWorld* InWorld)
 		DestActor->SetAsGlobalInWorld(true);
 		AddReplicator(InWorld, DestActor);
 	}
+#endif //!(UE_BUILD_SHIPPING || UE_BUILD_TEST)
 }
 
 void FGameplayDebugger::WorldDestroyed(UWorld* InWorld)
 {
+#if !(UE_BUILD_SHIPPING || UE_BUILD_TEST)
 	bool bIsServer = InWorld && InWorld->GetNetMode() < ENetMode::NM_Client; // (Only work on  server)
 	if (!bIsServer)
 	{
@@ -381,7 +384,9 @@ void FGameplayDebugger::WorldDestroyed(UWorld* InWorld)
 
 	// remove global replicator from level
 	AllReplilcatorsPerWorlds.Remove(InWorld);
+#endif //!(UE_BUILD_SHIPPING || UE_BUILD_TEST)
 }
+
 #if WITH_EDITOR
 void FGameplayDebugger::OnLevelActorAdded(AActor* InActor)
 {
@@ -431,6 +436,7 @@ void FGameplayDebugger::OnLevelActorDeleted(AActor* InActor)
 
 bool FGameplayDebugger::Exec(UWorld* Inworld, const TCHAR* Cmd, FOutputDevice& Ar)
 {
+#if !(UE_BUILD_SHIPPING || UE_BUILD_TEST)
 	bool bHandled = false;
 
 	
@@ -599,6 +605,9 @@ bool FGameplayDebugger::Exec(UWorld* Inworld, const TCHAR* Cmd, FOutputDevice& A
 	}
 
 	return bHandled;
+#else
+	return false;
+#endif //!(UE_BUILD_SHIPPING || UE_BUILD_TEST)
 }
 
 #undef LOCTEXT_NAMESPACE
diff --git a/Engine/Source/Developer/GameplayDebugger/Private/GameplayDebuggingComponent.cpp b/Engine/Source/Developer/GameplayDebugger/Private/GameplayDebuggingComponent.cpp
--- a/Engine/Source/Developer/GameplayDebugger/Private/GameplayDebuggingComponent.cpp
+++ b/Engine/Source/Developer/GameplayDebugger/Private/GameplayDebuggingComponent.cpp
@@ -149,7 +149,6 @@ UGameplayDebuggingComponent::UGameplayDebuggingComponent(const FObjectInitialize
 	bEnableClientEQSSceneProxy = false;
 	NextTargrtSelectionTime = 0;
 	ReplicateViewDataCounters.Init(0, EAIDebugDrawDataView::MAX);
-#endif
 
 	AGameplayDebuggingReplicator* Replicator = Cast<AGameplayDebuggingReplicator>(GetOwner());
 	if (Replicator)
@@ -161,7 +160,7 @@ UGameplayDebuggingComponent::UGameplayDebuggingComponent(const FObjectInitialize
 			ReplicateViewDataCounters[Index] = Settings.CheckFlag(Index);
 		}
 	}
-
+#endif
 }
 
 void UGameplayDebuggingComponent::Activate(bool bReset)
@@ -425,6 +424,7 @@ void UGameplayDebuggingComponent::CollectBasicData()
 
 void UGameplayDebuggingComponent::CollectBasicMovementData(APawn* MyPawn)
 {
+#if !(UE_BUILD_SHIPPING || UE_BUILD_TEST)
 	UCharacterMovementComponent* CharMovement = Cast<UCharacterMovementComponent>(MyPawn->GetMovementComponent());
 	if (CharMovement)
 	{
@@ -439,10 +439,12 @@ void UGameplayDebuggingComponent::CollectBasicMovementData(APawn* MyPawn)
 		MovementBaseInfo = TEXT("");
 		MovementModeInfo = TEXT("");
 	}
+#endif //!(UE_BUILD_SHIPPING || UE_BUILD_TEST)
 }
 
 void UGameplayDebuggingComponent::CollectBasicPathData(APawn* MyPawn)
 {
+#if !(UE_BUILD_SHIPPING || UE_BUILD_TEST)
 	UNavigationSystem* NavSys = UNavigationSystem::GetCurrent(MyPawn->GetWorld());
 	AAIController* MyAIController = Cast<AAIController>(MyPawn->GetController());
 
@@ -495,10 +497,12 @@ void UGameplayDebuggingComponent::CollectBasicPathData(APawn* MyPawn)
 	{
 		PathFollowingInfo = TEXT("");
 	}
+#endif //!(UE_BUILD_SHIPPING || UE_BUILD_TEST)
 }
 
 void UGameplayDebuggingComponent::CollectBasicBehaviorData(APawn* MyPawn)
 {
+#if !(UE_BUILD_SHIPPING || UE_BUILD_TEST)
 	AAIController* MyAIController = Cast<AAIController>(MyPawn->GetController());
 	UBehaviorTreeComponent* BTC = MyAIController ? Cast<UBehaviorTreeComponent>(MyAIController->BrainComponent) : nullptr;
 	bIsUsingBehaviorTree = (BTC != nullptr);
@@ -515,10 +519,12 @@ void UGameplayDebuggingComponent::CollectBasicBehaviorData(APawn* MyPawn)
 		CurrentAIState = TEXT("");
 		CurrentAIAssets = TEXT("");
 	}
+#endif //!(UE_BUILD_SHIPPING || UE_BUILD_TEST)
 }
 
 void UGameplayDebuggingComponent::CollectBasicAbilityData(APawn* MyPawn)
 {
+#if !(UE_BUILD_SHIPPING || UE_BUILD_TEST)
 	UAbilitySystemComponent* AbilityComp = MyPawn->FindComponentByClass<UAbilitySystemComponent>();
 	bIsUsingAbilities = (AbilityComp != nullptr);
 
@@ -550,12 +556,15 @@ void UGameplayDebuggingComponent::CollectBasicAbilityData(APawn* MyPawn)
 	{
 		AbilityInfo = TEXT("None");
 	}
+#endif //!(UE_BUILD_SHIPPING || UE_BUILD_TEST)
 }
 
 void UGameplayDebuggingComponent::CollectBasicAnimationData(APawn* MyPawn)
 {
+#if !(UE_BUILD_SHIPPING || UE_BUILD_TEST)
 	ACharacter* MyChar = Cast<ACharacter>(MyPawn);
 	MontageInfo = MyChar ? GetNameSafe(MyChar->GetCurrentMontage()) : TEXT("");
+#endif //!(UE_BUILD_SHIPPING || UE_BUILD_TEST)
 }
 
 void UGameplayDebuggingComponent::CollectBehaviorTreeData()
@@ -796,6 +805,7 @@ void UGameplayDebuggingComponent::ServerReplicateData(uint32 InMessage, uint32
 //////////////////////////////////////////////////////////////////////////
 void UGameplayDebuggingComponent::OnChangeEQSQuery()
 {
+#if !(UE_BUILD_SHIPPING || UE_BUILD_TEST)
 	AGameplayDebuggingReplicator* Replicator = Cast<AGameplayDebuggingReplicator>(GetOwner());
 	if (++CurrentEQSIndex >= EQSLocalData.Num())
 	{
@@ -804,6 +814,7 @@ void UGameplayDebuggingComponent::OnChangeEQSQuery()
 
 	UpdateBounds();
 	MarkRenderStateDirty();
+#endif //!(UE_BUILD_SHIPPING || UE_BUILD_TEST)
 }
 
 const FEnvQueryResult* UGameplayDebuggingComponent::GetQueryResult() const
@@ -939,9 +950,11 @@ void UGameplayDebuggingComponent::CollectEQSData()
 //----------------------------------------------------------------------//
 void UGameplayDebuggingComponent::OnRep_UpdateNavmesh()
 {
+#if !(UE_BUILD_SHIPPING || UE_BUILD_TEST)
 	NavMeshBounds = FBox(FVector(-HALF_WORLD_MAX, -HALF_WORLD_MAX, -HALF_WORLD_MAX), FVector(HALF_WORLD_MAX, HALF_WORLD_MAX, HALF_WORLD_MAX));
 	UpdateBounds();
 	MarkRenderStateDirty();
+#endif //!(UE_BUILD_SHIPPING || UE_BUILD_TEST)
 }
 
 bool UGameplayDebuggingComponent::ServerCollectNavmeshData_Validate(FVector_NetQuantize10 TargetLocation)
@@ -1086,6 +1099,7 @@ FORCEINLINE bool LineInCorrectDistance(const FVector& PlayerLoc, const FVector&
 #if WITH_RECAST
 ARecastNavMesh* UGameplayDebuggingComponent::GetNavData()
 {
+#if !(UE_BUILD_SHIPPING || UE_BUILD_TEST)
 	UNavigationSystem* NavSys = UNavigationSystem::GetCurrent(GetWorld());
 	if (NavSys == NULL)
 	{
@@ -1102,6 +1116,9 @@ ARecastNavMesh* UGameplayDebuggingComponent::GetNavData()
 
 	// If it wasn't found, just get the main nav-mesh data.
 	return Cast<ARecastNavMesh>(NavSys->GetMainNavData(FNavigationSystem::DontCreate));
+#else
+	return NULL;
+#endif //!(UE_BUILD_SHIPPING || UE_BUILD_TEST)
 }
 #endif
 
@@ -1382,6 +1399,7 @@ public:
 FPrimitiveSceneProxy* UGameplayDebuggingComponent::CreateSceneProxy()
 {
 	FDebugRenderSceneCompositeProxy* CompositeProxy = NULL;
+#if !(UE_BUILD_SHIPPING || UE_BUILD_TEST)
 	AGameplayDebuggingReplicator* Replicator = Cast<AGameplayDebuggingReplicator>(GetOwner());
 	if (!World || World->GetNetMode() == NM_DedicatedServer)
 	{
@@ -1483,6 +1501,7 @@ FPrimitiveSceneProxy* UGameplayDebuggingComponent::CreateSceneProxy()
 		CompositeProxy->AddChild(DebugSceneProxy);
 	}
 
+#endif //!(UE_BUILD_SHIPPING || UE_BUILD_TEST)
 	return CompositeProxy;
 }
 
diff --git a/Engine/Source/Developer/GameplayDebugger/Private/GameplayDebuggingControllerComponent.cpp b/Engine/Source/Developer/GameplayDebugger/Private/GameplayDebuggingControllerComponent.cpp
--- a/Engine/Source/Developer/GameplayDebugger/Private/GameplayDebuggingControllerComponent.cpp
+++ b/Engine/Source/Developer/GameplayDebugger/Private/GameplayDebuggingControllerComponent.cpp
@@ -17,6 +17,7 @@ UGameplayDebuggingControllerComponent::UGameplayDebuggingControllerComponent(con
 	: Super(ObjectInitializer)
 	, KeyPressActivationTime(0.4f)
 {
+#if !(UE_BUILD_SHIPPING || UE_BUILD_TEST)
 	PrimaryComponentTick.bCanEverTick = true;
 	bWantsInitializeComponent = true;
 	bAutoActivate = false;
@@ -30,13 +31,16 @@ UGameplayDebuggingControllerComponent::UGameplayDebuggingControllerComponent(con
 
 	ControlKeyPressedTime = 0;
 	ActivationKey = FInputChord(EKeys::Apostrophe, false, false, false, false);
+#endif //!(UE_BUILD_SHIPPING || UE_BUILD_TEST)
 }
 
 void UGameplayDebuggingControllerComponent::OnRegister()
 {
 	Super::OnRegister();
+#if !(UE_BUILD_SHIPPING || UE_BUILD_TEST)
 	BindActivationKeys();
 	SetActiveViews(GameplayDebuggerSettings().DebuggerShowFlags);
+#endif
 }
 
 AGameplayDebuggingReplicator* UGameplayDebuggingControllerComponent::GetDebuggingReplicator() const
@@ -46,6 +50,7 @@ AGameplayDebuggingReplicator* UGameplayDebuggingControllerComponent::GetDebuggin
 
 void UGameplayDebuggingControllerComponent::BeginDestroy()
 {
+#if !(UE_BUILD_SHIPPING || UE_BUILD_TEST)
 	AHUD* GameHUD = PlayerOwner.IsValid() ? PlayerOwner->GetHUD() : NULL;
 	if (GameHUD)
 	{
@@ -78,6 +83,7 @@ void UGameplayDebuggingControllerComponent::BeginDestroy()
 		}
 		GetDebuggingReplicator()->ServerReplicateMessage(DebugAITargetActor, EDebugComponentMessage::DisableExtendedView);
 	}
+#endif //!(UE_BUILD_SHIPPING || UE_BUILD_TEST)
 
 	Super::BeginDestroy();
 }
@@ -94,6 +100,7 @@ uint32 UGameplayDebuggingControllerComponent::GetActiveViews()
 
 void UGameplayDebuggingControllerComponent::SetActiveViews(uint32 InActiveViews)
 {
+#if !(UE_BUILD_SHIPPING || UE_BUILD_TEST)
 	GameplayDebuggerSettings(GetDebuggingReplicator()).DebuggerShowFlags = InActiveViews;
 
 	if (GetDebuggingReplicator())
@@ -104,10 +111,12 @@ void UGameplayDebuggingControllerComponent::SetActiveViews(uint32 InActiveViews)
 			EnableActiveView(CurrentView, IsViewActive(CurrentView));
 		}
 	}
+#endif //!(UE_BUILD_SHIPPING || UE_BUILD_TEST)
 }
 
 void UGameplayDebuggingControllerComponent::EnableActiveView(EAIDebugDrawDataView::Type View, bool bEnable)
 {
+#if !(UE_BUILD_SHIPPING || UE_BUILD_TEST)
 	bEnable ? GameplayDebuggerSettings(GetDebuggingReplicator()).SetFlag(View) : GameplayDebuggerSettings(GetDebuggingReplicator()).ClearFlag(View);
 
 	if (GetDebuggingReplicator())
@@ -120,10 +129,12 @@ void UGameplayDebuggingControllerComponent::EnableActiveView(EAIDebugDrawDataVie
 		}
 #endif // WITH_EQS
 	}
+#endif //!(UE_BUILD_SHIPPING || UE_BUILD_TEST)
 }
 
 void UGameplayDebuggingControllerComponent::BindActivationKeys()
 {
+#if !(UE_BUILD_SHIPPING || UE_BUILD_TEST)
 	if (PlayerOwner.IsValid() && PlayerOwner->InputComponent && PlayerOwner->PlayerInput)
 	{
 		// find current activation key used for 'EnableGDT' binding
@@ -142,10 +153,12 @@ void UGameplayDebuggingControllerComponent::BindActivationKeys()
 		PlayerOwner->InputComponent->BindKey(ActivationKey, IE_Pressed, this, &UGameplayDebuggingControllerComponent::OnActivationKeyPressed);
 		PlayerOwner->InputComponent->BindKey(ActivationKey, IE_Released, this, &UGameplayDebuggingControllerComponent::OnActivationKeyReleased);
 	}
+#endif //!(UE_BUILD_SHIPPING || UE_BUILD_TEST)
 }
 
 void UGameplayDebuggingControllerComponent::OnActivationKeyPressed()
 {
+#if !(UE_BUILD_SHIPPING || UE_BUILD_TEST)
 	if (GetDebuggingReplicator() && PlayerOwner.IsValid())
 	{
 		if (!bToolActivated)
@@ -161,10 +174,12 @@ void UGameplayDebuggingControllerComponent::OnActivationKeyPressed()
 		ControlKeyPressedTime = GetWorld()->GetTimeSeconds();
 		EnableTargetSelection(true);
 	}
+#endif //!(UE_BUILD_SHIPPING || UE_BUILD_TEST)
 }
 
 void UGameplayDebuggingControllerComponent::OnActivationKeyReleased()
 {
+#if !(UE_BUILD_SHIPPING || UE_BUILD_TEST)
 	const float KeyPressedTime = GetWorld()->GetTimeSeconds() - ControlKeyPressedTime;
 
 	EnableTargetSelection(false);
@@ -183,19 +198,22 @@ void UGameplayDebuggingControllerComponent::OnActivationKeyReleased()
 			}
 		}
 	}
+#endif //!(UE_BUILD_SHIPPING || UE_BUILD_TEST)
 }
 
 void UGameplayDebuggingControllerComponent::OpenDebugTool()
 {
+#if !(UE_BUILD_SHIPPING || UE_BUILD_TEST)
 	if (GetDebuggingReplicator() && IsActive())
 	{
 		bToolActivated = true;
 	}
-
+#endif //!(UE_BUILD_SHIPPING || UE_BUILD_TEST)
 }
 
 void UGameplayDebuggingControllerComponent::CloseDebugTool()
 {
+#if !(UE_BUILD_SHIPPING || UE_BUILD_TEST)
 	if (GetDebuggingReplicator())
 	{
 		Deactivate();
@@ -205,11 +223,14 @@ void UGameplayDebuggingControllerComponent::CloseDebugTool()
 		GetDebuggingReplicator()->ServerReplicateMessage(NULL, EDebugComponentMessage::DeactivateReplilcation, EAIDebugDrawDataView::Empty);
 		bToolActivated = false;
 	}
+#endif //!(UE_BUILD_SHIPPING || UE_BUILD_TEST)
 }
 
 void UGameplayDebuggingControllerComponent::EnableTargetSelection(bool bEnable)
 {
+#if !(UE_BUILD_SHIPPING || UE_BUILD_TEST)
 	GetDebuggingReplicator()->ClientEnableTargetSelection(bEnable, PlayerOwner.Get());
+#endif //!(UE_BUILD_SHIPPING || UE_BUILD_TEST)
 }
 
 void UGameplayDebuggingControllerComponent::BindAIDebugViewKeys()
@@ -335,6 +356,7 @@ void UGameplayDebuggingControllerComponent::TickComponent(float DeltaTime, enum
 {
 	Super::TickComponent(DeltaTime, TickType, ThisTickFunction);
 
+#if !(UE_BUILD_SHIPPING || UE_BUILD_TEST)
 	if (GetWorld()->GetTimeSeconds() - ControlKeyPressedTime > KeyPressActivationTime && !bToolActivated)
 	{
 		OpenDebugTool();
@@ -344,15 +366,21 @@ void UGameplayDebuggingControllerComponent::TickComponent(float DeltaTime, enum
 	{
 		return;
 	}
+#endif //!(UE_BUILD_SHIPPING || UE_BUILD_TEST)
 }
 
 float UGameplayDebuggingControllerComponent::GetUpdateNavMeshTimeRemaining() const
 {
+#if !(UE_BUILD_SHIPPING || UE_BUILD_TEST)
 	return GetWorld()->GetTimerManager().GetTimerRemaining(TimerHandle_UpdateNavMeshTimer);
+#else
+	return FLT_MAX;
+#endif
 }
 
 void UGameplayDebuggingControllerComponent::UpdateNavMeshTimer()
 {
+#if !(UE_BUILD_SHIPPING || UE_BUILD_TEST)
 	const APawn* PlayerPawn = PlayerOwner.IsValid() ? PlayerOwner->GetPawnOrSpectator() : NULL;
 	UGameplayDebuggingComponent* DebuggingComponent = GetDebuggingReplicator() ? GetDebuggingReplicator()->GetDebugComponent() : NULL;
 	if (DebuggingComponent)
@@ -371,4 +399,5 @@ void UGameplayDebuggingControllerComponent::UpdateNavMeshTimer()
 			DebuggingComponent->ServerCollectNavmeshData(AdditionalTargetLoc);
 		}
 	}
+#endif //!(UE_BUILD_SHIPPING || UE_BUILD_TEST)
 }
diff --git a/Engine/Source/Developer/GameplayDebugger/Private/GameplayDebuggingHUDComponent.cpp b/Engine/Source/Developer/GameplayDebugger/Private/GameplayDebuggingHUDComponent.cpp
--- a/Engine/Source/Developer/GameplayDebugger/Private/GameplayDebuggingHUDComponent.cpp
+++ b/Engine/Source/Developer/GameplayDebugger/Private/GameplayDebuggingHUDComponent.cpp
@@ -33,27 +33,32 @@ AGameplayDebuggingHUDComponent::AGameplayDebuggingHUDComponent(const FObjectInit
 	: Super(ObjectInitializer)
 	, EngineShowFlags(EShowFlagInitMode::ESFIM_Game)
 {
+#if !(UE_BUILD_SHIPPING || UE_BUILD_TEST)
 	World = NULL;
 	SetTickableWhenPaused(true);
 
-#if WITH_EDITORONLY_DATA
+#	if WITH_EDITORONLY_DATA
 	SetIsTemporarilyHiddenInEditor(true);
 	SetActorHiddenInGame(false);
 	bHiddenEdLevel = true;
 	bHiddenEdLayer = true;
 	bHiddenEd = true;
 	bEditable = false;
+#	endif
 #endif
 }
 
 void AGameplayDebuggingHUDComponent::Render()
 {
+#if !(UE_BUILD_SHIPPING || UE_BUILD_TEST)
 	EngineShowFlags = Canvas && Canvas->SceneView && Canvas->SceneView->Family ? Canvas->SceneView->Family->EngineShowFlags : FEngineShowFlags(GIsEditor ? EShowFlagInitMode::ESFIM_Editor : EShowFlagInitMode::ESFIM_Game);
 	PrintAllData();
+#endif
 }
 
 AGameplayDebuggingReplicator* AGameplayDebuggingHUDComponent::GetDebuggingReplicator()
 {
+#if !(UE_BUILD_SHIPPING || UE_BUILD_TEST)
 	if (CachedDebuggingReplicator.IsValid() && CachedDebuggingReplicator->GetLocalPlayerOwner() == PlayerOwner)
 	{
 		return CachedDebuggingReplicator.Get();
@@ -68,17 +73,19 @@ AGameplayDebuggingReplicator* AGameplayDebuggingHUDComponent::GetDebuggingReplic
 			return Replicator;
 		}
 	}
-
+#endif
 	return NULL;
 }
 
 void AGameplayDebuggingHUDComponent::GetKeyboardDesc(TArray<FDebugCategoryView>& Categories) 
 {
+#if !(UE_BUILD_SHIPPING || UE_BUILD_TEST)
 	Categories.Add(FDebugCategoryView(EAIDebugDrawDataView::NavMesh, TEXT("NavMesh")));
 	Categories.Add(FDebugCategoryView(EAIDebugDrawDataView::Basic, TEXT("Basic")));
 	Categories.Add(FDebugCategoryView(EAIDebugDrawDataView::BehaviorTree, TEXT("BehaviorTree")));
 	Categories.Add(FDebugCategoryView(EAIDebugDrawDataView::EQS, TEXT("EQS")));
 	Categories.Add(FDebugCategoryView(EAIDebugDrawDataView::Perception, TEXT("Perception")));
+#endif //!(UE_BUILD_SHIPPING || UE_BUILD_TEST)
 }
 
 void AGameplayDebuggingHUDComponent::PrintAllData()
@@ -121,6 +128,7 @@ void AGameplayDebuggingHUDComponent::PrintAllData()
 
 void AGameplayDebuggingHUDComponent::DrawMenu(const float X, const float Y, class UGameplayDebuggingComponent* DebugComponent)
 {
+#if !(UE_BUILD_SHIPPING || UE_BUILD_TEST)
 	const float OldX = DefaultContext.CursorX;
 	const float OldY = DefaultContext.CursorY;
 
@@ -190,6 +198,7 @@ void AGameplayDebuggingHUDComponent::DrawMenu(const float X, const float Y, clas
 
 	DefaultContext.CursorX = OldX;
 	DefaultContext.CursorY = OldY;
+#endif //!(UE_BUILD_SHIPPING || UE_BUILD_TEST)
 }
 
 void AGameplayDebuggingHUDComponent::DrawDebugComponentData(APlayerController* MyPC, class UGameplayDebuggingComponent *DebugComponent)
@@ -684,6 +693,7 @@ void AGameplayDebuggingHUDComponent::DrawNavMeshSnapshot(APlayerController* PC,
 //////////////////////////////////////////////////////////////////////////
 void AGameplayDebuggingHUDComponent::PrintString(FPrintContext& Context, const FString& InString )
 {
+#if !(UE_BUILD_SHIPPING || UE_BUILD_TEST)
 	class FParserHelper
 	{
 		enum TokenType
@@ -887,15 +897,19 @@ void AGameplayDebuggingHUDComponent::PrintString(FPrintContext& Context, const F
 	}
 
 	Context.CursorY += YMovement;
+#endif //!(UE_BUILD_SHIPPING || UE_BUILD_TEST)
 }
 
 void AGameplayDebuggingHUDComponent::PrintString(FPrintContext& Context, const FColor& InColor, const FString& InString )
 {
-	PrintString( Context, FString::Printf(TEXT("{%s}%s"), *InColor.ToString(), *InString));
+#if !(UE_BUILD_SHIPPING || UE_BUILD_TEST)
+	PrintString(Context, FString::Printf(TEXT("{%s}%s"), *InColor.ToString(), *InString));
+#endif //!(UE_BUILD_SHIPPING || UE_BUILD_TEST)
 }
 
 void AGameplayDebuggingHUDComponent::PrintString(FPrintContext& Context, const FColor& InColor, const FString& InString, float X, float Y )
 {
+#if !(UE_BUILD_SHIPPING || UE_BUILD_TEST)
 	const float OldX = Context.CursorX, OldY = Context.CursorY;
 	const float OldDX = Context.DefaultX, OldDY = Context.DefaultY;
 
@@ -907,16 +921,18 @@ void AGameplayDebuggingHUDComponent::PrintString(FPrintContext& Context, const F
 	Context.CursorY = OldY;
 	Context.DefaultX = OldDX;
 	Context.DefaultY = OldDY;
+#endif
 }
 
 void AGameplayDebuggingHUDComponent::CalulateStringSize(const AGameplayDebuggingHUDComponent::FPrintContext& DefaultContext, UFont* Font, const FString& InString, float& OutX, float& OutY )
 {
+#if !(UE_BUILD_SHIPPING || UE_BUILD_TEST)
 	FString String = InString;
 	const FRegexPattern ElementRegexPattern(TEXT("\\{.*?\\}"));
 	FRegexMatcher ElementRegexMatcher(ElementRegexPattern, String);
 
 	while (ElementRegexMatcher.FindNext())
-{
+	{
 		int32 AttributeListBegin = ElementRegexMatcher.GetCaptureGroupBeginning(0);
 		int32 AttributeListEnd = ElementRegexMatcher.GetCaptureGroupEnding(0);
 		String.RemoveAt(AttributeListBegin, AttributeListEnd - AttributeListBegin);
@@ -928,37 +944,45 @@ void AGameplayDebuggingHUDComponent::CalulateStringSize(const AGameplayDebugging
 	{
 		DefaultContext.Canvas->StrLen(Font != NULL ? Font : DefaultContext.Font, String, OutX, OutY);
 	}
+#endif
 }
 
 void AGameplayDebuggingHUDComponent::CalulateTextSize(const AGameplayDebuggingHUDComponent::FPrintContext& DefaultContext, UFont* Font, const FText& InText, float& OutX, float& OutY)
 {
+#if !(UE_BUILD_SHIPPING || UE_BUILD_TEST)
 	CalulateStringSize(DefaultContext, Font, InText.ToString(), OutX, OutY);
+#endif
 }
 
 FVector AGameplayDebuggingHUDComponent::ProjectLocation(const AGameplayDebuggingHUDComponent::FPrintContext& Context, const FVector& Location)
 {
+#if !(UE_BUILD_SHIPPING || UE_BUILD_TEST)
 	if (Context.Canvas != NULL)
 	{
 		return Context.Canvas->Project(Location);
 	}
-
+#endif
 	return FVector();
 }
 
 void AGameplayDebuggingHUDComponent::DrawItem(const AGameplayDebuggingHUDComponent::FPrintContext& DefaultContext, class FCanvasItem& Item, float X, float Y )
 {
-	if(DefaultContext.Canvas)
+#if !(UE_BUILD_SHIPPING || UE_BUILD_TEST)
+	if (DefaultContext.Canvas)
 	{
 		DefaultContext.Canvas->DrawItem( Item, FVector2D( X, Y ) );
 	}
+#endif
 }
 
 void AGameplayDebuggingHUDComponent::DrawIcon(const AGameplayDebuggingHUDComponent::FPrintContext& DefaultContext, const FColor& InColor, const FCanvasIcon& Icon, float X, float Y, float Scale)
 {
-	if(DefaultContext.Canvas)
+#if !(UE_BUILD_SHIPPING || UE_BUILD_TEST)
+	if (DefaultContext.Canvas)
 	{
 		DefaultContext.Canvas->SetDrawColor(InColor);
 		DefaultContext.Canvas->DrawIcon(Icon, X, Y, Scale);
 	}
+#endif
 }
 
diff --git a/Engine/Source/Developer/GameplayDebugger/Private/GameplayDebuggingReplicator.cpp b/Engine/Source/Developer/GameplayDebugger/Private/GameplayDebuggingReplicator.cpp
--- a/Engine/Source/Developer/GameplayDebugger/Private/GameplayDebuggingReplicator.cpp
+++ b/Engine/Source/Developer/GameplayDebugger/Private/GameplayDebuggingReplicator.cpp
@@ -23,6 +23,7 @@ AGameplayDebuggingReplicator::AGameplayDebuggingReplicator(const FObjectInitiali
 	, LastDrawAtFrame(0)
 	, PlayerControllersUpdateDelay(0)
 {
+#if !(UE_BUILD_SHIPPING || UE_BUILD_TEST)
 	// Structure to hold one-time initialization
 	struct FConstructorStatics
 	{
@@ -87,6 +88,7 @@ AGameplayDebuggingReplicator::AGameplayDebuggingReplicator(const FObjectInitiali
 
 		AGameplayDebuggingReplicator::OnSelectionChangedDelegate.AddUObject(this, &AGameplayDebuggingReplicator::ServerSetActorToDebug);
 	}
+#endif //!(UE_BUILD_SHIPPING || UE_BUILD_TEST)
 }
 
 void AGameplayDebuggingReplicator::GetLifetimeReplicatedProps(TArray< FLifetimeProperty > & OutLifetimeProps) const
@@ -109,7 +111,9 @@ bool AGameplayDebuggingReplicator::IsNetRelevantFor(const APlayerController* Rea
 void AGameplayDebuggingReplicator::PostInitializeComponents()
 {
 	Super::PostInitializeComponents();
+#if !(UE_BUILD_SHIPPING || UE_BUILD_TEST)
 	SetActorTickEnabled(true);
+#endif
 }
 
 #if WITH_EDITOR
@@ -117,6 +121,7 @@ void AGameplayDebuggingReplicator::PostEditChangeProperty(FPropertyChangedEvent&
 {
 	Super::PostEditChangeProperty(PropertyChangedEvent);
 
+#if !(UE_BUILD_SHIPPING || UE_BUILD_TEST)
 	if (!PropertyChangedEvent.Property)
 	{
 		return;
@@ -155,15 +160,17 @@ void AGameplayDebuggingReplicator::PostEditChangeProperty(FPropertyChangedEvent&
 		GetDebugComponent()->SetEQSIndex(ActiveEQSIndex);
 	}
 #endif // WITH_EQS
-
 #undef CHECK_AND_UPDATE_FLAGS
+
+#endif //!(UE_BUILD_SHIPPING || UE_BUILD_TEST)
 }
-#endif
+#endif //WITH_EDITOR
 
 void AGameplayDebuggingReplicator::BeginPlay()
 {
 	Super::BeginPlay();
 
+#if !(UE_BUILD_SHIPPING || UE_BUILD_TEST)
 	if (Role == ROLE_Authority)
 	{
 		bReplicates = false;
@@ -266,39 +273,47 @@ void AGameplayDebuggingReplicator::BeginPlay()
 		OnRep_AutoActivate();
 	}
 
+#endif //!(UE_BUILD_SHIPPING || UE_BUILD_TEST)
 }
 
 void AGameplayDebuggingReplicator::PostNetInit()
 {
 	Super::PostNetInit();
+#if !(UE_BUILD_SHIPPING || UE_BUILD_TEST)
 	if (bAutoActivate)
 	{
 		OnRep_AutoActivate();
 	}
+#endif
 }
 
 void AGameplayDebuggingReplicator::ClientAutoActivate_Implementation()
 {
+#if !(UE_BUILD_SHIPPING || UE_BUILD_TEST)
 	// we are already replicated so let's activate tool
 	if (GetWorld() && GetNetMode() == ENetMode::NM_Client && !IsToolCreated() && !IsGlobalInWorld())
 	{
 		CreateTool();
 		EnableTool();
 	}
+#endif
 }
 
 void AGameplayDebuggingReplicator::OnRep_AutoActivate()
 {
+#if !(UE_BUILD_SHIPPING || UE_BUILD_TEST)
 	// we are already replicated so let's activate tool
 	if (GetWorld() && GetNetMode() == ENetMode::NM_Client && !IsToolCreated() && !IsGlobalInWorld())
 	{
 		CreateTool();
 		EnableTool();
 	}
+#endif
 }
 
 UGameplayDebuggingComponent* AGameplayDebuggingReplicator::GetDebugComponent()
 {
+#if !(UE_BUILD_SHIPPING || UE_BUILD_TEST)
 	if (IsPendingKill() == false && !DebugComponent && DebugComponentClass.IsValid() && GetNetMode() < ENetMode::NM_Client)
 	{
 		DebugComponent = NewObject<UGameplayDebuggingComponent>(this, DebugComponentClass.Get());
@@ -306,17 +321,18 @@ UGameplayDebuggingComponent* AGameplayDebuggingReplicator::GetDebugComponent()
 		DebugComponent->RegisterComponent();
 		DebugComponent->Activate();
 	}
-
+#endif
 	return DebugComponent;
 }
 
 class UNetConnection* AGameplayDebuggingReplicator::GetNetConnection()
 {
+#if !(UE_BUILD_SHIPPING || UE_BUILD_TEST)
 	if (LocalPlayerOwner && LocalPlayerOwner->IsPendingKill() == false && IsPendingKill() == false)
 	{
 		return LocalPlayerOwner->GetNetConnection();
 	}
-
+#endif
 	return NULL;
 }
 
@@ -327,10 +343,12 @@ bool AGameplayDebuggingReplicator::ClientEnableTargetSelection_Validate(bool, AP
 
 void AGameplayDebuggingReplicator::ClientEnableTargetSelection_Implementation(bool bEnable, APlayerController* Context)
 {
+#if !(UE_BUILD_SHIPPING || UE_BUILD_TEST)
 	if (GetDebugComponent())
 	{
 		GetDebugComponent()->ClientEnableTargetSelection(bEnable);
 	}
+#endif
 }
 
 bool AGameplayDebuggingReplicator::ClientReplicateMessage_Validate(class AActor* Actor, uint32 InMessage, uint32 DataView)
@@ -366,11 +384,16 @@ void AGameplayDebuggingReplicator::ServerReplicateMessage_Implementation(class
 
 bool AGameplayDebuggingReplicator::IsDrawEnabled()
 {
+#if !(UE_BUILD_SHIPPING || UE_BUILD_TEST)
 	return bEnabledDraw && GetWorld() && GetNetMode() != ENetMode::NM_DedicatedServer;
+#else
+	return false;
+#endif
 }
 
 void AGameplayDebuggingReplicator::EnableDraw(bool bEnable)
 {
+#if !(UE_BUILD_SHIPPING || UE_BUILD_TEST)
 	bEnabledDraw = bEnable;
 
 	if (AHUD* const GameHUD = LocalPlayerOwner ? LocalPlayerOwner->GetHUD() : NULL)
@@ -385,6 +408,7 @@ void AGameplayDebuggingReplicator::EnableDraw(bool bEnable)
 		DebugComponent->EnableClientEQSSceneProxy(bEnable && bEnabledEQSView ? true : false);
 		DebugComponent->MarkRenderStateDirty();
 	}
+#endif
 }
 
 bool AGameplayDebuggingReplicator::IsToolCreated()
@@ -395,6 +419,7 @@ bool AGameplayDebuggingReplicator::IsToolCreated()
 
 void AGameplayDebuggingReplicator::CreateTool()
 {
+#if !(UE_BUILD_SHIPPING || UE_BUILD_TEST)
 	if (GetWorld() && GetNetMode() != ENetMode::NM_DedicatedServer)
 	{
 		UGameplayDebuggingControllerComponent*  GDC = FindComponentByClass<UGameplayDebuggingControllerComponent>();
@@ -410,10 +435,12 @@ void AGameplayDebuggingReplicator::CreateTool()
 			GDC->RegisterComponent();
 		}
 	}
+#endif
 }
 
 void AGameplayDebuggingReplicator::EnableTool()
 {
+#if !(UE_BUILD_SHIPPING || UE_BUILD_TEST)
 	if (GetWorld() && GetNetMode() != ENetMode::NM_DedicatedServer)
 	{
 		UGameplayDebuggingControllerComponent*  GDC = FindComponentByClass<UGameplayDebuggingControllerComponent>();
@@ -424,12 +451,14 @@ void AGameplayDebuggingReplicator::EnableTool()
 			GDC->OnActivationKeyReleased();
 		}
 	}
+#endif
 }
 
 void AGameplayDebuggingReplicator::TickActor(float DeltaTime, enum ELevelTick TickType, FActorTickFunction& ThisTickFunction)
 {
 	Super::TickActor(DeltaTime, TickType, ThisTickFunction);
 
+#if !(UE_BUILD_SHIPPING || UE_BUILD_TEST)
 	UWorld* World = GetWorld();
 	if (!IsGlobalInWorld() || !World || GetNetMode() == ENetMode::NM_Client || !IGameplayDebugger::IsAvailable())
 	{
@@ -457,6 +486,7 @@ void AGameplayDebuggingReplicator::TickActor(float DeltaTime, enum ELevelTick Ti
 		}
 		PlayerControllersUpdateDelay = 5;
 	}
+#endif
 }
 
 bool AGameplayDebuggingReplicator::ServerSetActorToDebug_Validate(AActor* InActor)
@@ -466,6 +496,7 @@ bool AGameplayDebuggingReplicator::ServerSetActorToDebug_Validate(AActor* InActo
 
 void AGameplayDebuggingReplicator::ServerSetActorToDebug_Implementation(AActor* InActor)
 { 
+#if !(UE_BUILD_SHIPPING || UE_BUILD_TEST)
 	if (LastSelectedActorToDebug != InActor)
 	{
 		LastSelectedActorToDebug = InActor;
@@ -481,11 +512,12 @@ void AGameplayDebuggingReplicator::ServerSetActorToDebug_Implementation(AActor*
 	{
 		DebugComponent->SetActorToDebug(InActor);
 	}
+#endif
 }
 
 void AGameplayDebuggingReplicator::OnDebugAIDelegate(class UCanvas* Canvas, class APlayerController* PC)
 {
-#if WITH_EDITOR
+#if WITH_EDITOR && !(UE_BUILD_SHIPPING || UE_BUILD_TEST)
 	if (!GIsEditor)
 	{
 		return;
@@ -612,10 +644,12 @@ void AGameplayDebuggingReplicator::DrawDebugDataDelegate(class UCanvas* Canvas,
 	}
 
 	DrawDebugData(Canvas, PC);
+#endif
 }
 
 void AGameplayDebuggingReplicator::DrawDebugData(class UCanvas* Canvas, class APlayerController* PC)
 {
+#if !(UE_BUILD_SHIPPING || UE_BUILD_TEST)
 	if (!LocalPlayerOwner)
 	{
 		return;
@@ -651,6 +685,7 @@ void AGameplayDebuggingReplicator::DrawDebugData(class UCanvas* Canvas, class AP
 
 void AGameplayDebuggingReplicator::DebugNextPawn(UClass* CompareClass, APawn* CurrentPawn)
 {
+#if !(UE_BUILD_SHIPPING || UE_BUILD_TEST)
 	APawn* LastSeen = nullptr;
 	APawn* FirstSeen = nullptr;
 	// Search through the list looking for the next one of this type
@@ -676,10 +711,12 @@ void AGameplayDebuggingReplicator::DebugNextPawn(UClass* CompareClass, APawn* Cu
 	{
 		ServerSetActorToDebug(FirstSeen);
 	}
+#endif
 }
 
 void AGameplayDebuggingReplicator::DebugPrevPawn(UClass* CompareClass, APawn* CurrentPawn)
 {
+#if !(UE_BUILD_SHIPPING || UE_BUILD_TEST)
 	APawn* LastSeen = nullptr;
 	APawn* PrevSeen = nullptr;
 	APawn* FirstSeen = nullptr;
@@ -717,4 +754,5 @@ void AGameplayDebuggingReplicator::DebugPrevPawn(UClass* CompareClass, APawn* Cu
 	{
 		ServerSetActorToDebug(PrevSeen);
 	}
+#endif
 }
