Fixed GameplayDebugger crash on debugging a pawn in absence of navigation data #UE4

diff --git a/Engine/Source/Developer/GameplayDebugger/Private/GameplayDebuggingComponent.cpp b/Engine/Source/Developer/GameplayDebugger/Private/GameplayDebuggingComponent.cpp
--- a/Engine/Source/Developer/GameplayDebugger/Private/GameplayDebuggingComponent.cpp
+++ b/Engine/Source/Developer/GameplayDebugger/Private/GameplayDebuggingComponent.cpp
@@ -449,7 +449,10 @@ void UGameplayDebuggingComponent::CollectBasicPathData(APawn* MyPawn)
 	AAIController* MyAIController = Cast<AAIController>(MyPawn->GetController());
 
 	const ANavigationData* NavData = NavSys->GetNavDataForProps(MyAIController->GetNavAgentPropertiesRef());
-	NavDataInfo = NavData->GetConfig().Name.ToString();
+	if (NavData)
+	{
+		NavDataInfo = NavData->GetConfig().Name.ToString();
+	}
 
 	UPathFollowingComponent* PFC = MyAIController->GetPathFollowingComponent();
 	bIsUsingPathFollowing = (PFC != nullptr);
