Added debug camera to gameplay debugger, to have a way to fly around (Tab key to switch between cameras).
Added alternate keyboard bindings to gameplay debugger, for keyboards without numpad (can be enabled in gameplay debugger settings). It uses Alt + [regular number] shortcuts.

diff --git a/Engine/Source/Developer/GameplayDebugger/Classes/GameplayDebuggerSettings.h b/Engine/Source/Developer/GameplayDebugger/Classes/GameplayDebuggerSettings.h
--- a/Engine/Source/Developer/GameplayDebugger/Classes/GameplayDebuggerSettings.h
+++ b/Engine/Source/Developer/GameplayDebugger/Classes/GameplayDebuggerSettings.h
@@ -63,6 +63,8 @@ public:
 		return Settings;
 	}
 
+	bool UseAlternateKeys() { return bUseAlternateKeys; }
+
 	const FGDTCustomViewNames& GetCustomViewNames() { return CustomViewNames; }
 
 #if ADD_LEVEL_EDITOR_EXTENSIONS
@@ -111,6 +113,9 @@ protected:
 	UPROPERTY(EditAnywhere, config, Category = "GameplayDebugger")
 	bool GameView5;
 
+	UPROPERTY(EditAnywhere, config, Category = "GameplayDebugger")
+	bool bUseAlternateKeys;
+
 	// Holds an event delegate that is executed when a setting has changed.
 	FSettingChangedEvent SettingChangedEvent;
 };
diff --git a/Engine/Source/Developer/GameplayDebugger/Classes/GameplayDebuggingControllerComponent.h b/Engine/Source/Developer/GameplayDebugger/Classes/GameplayDebuggingControllerComponent.h
--- a/Engine/Source/Developer/GameplayDebugger/Classes/GameplayDebuggingControllerComponent.h
+++ b/Engine/Source/Developer/GameplayDebugger/Classes/GameplayDebuggingControllerComponent.h
@@ -8,6 +8,7 @@
 #include "TimerManager.h"
 #include "GameplayDebuggingTypes.h"
 #include "GameplayDebugger.h"
+#include "GameFramework/HUD.h"
 #include "GameplayDebuggingControllerComponent.generated.h"
 
 class AGameplayDebuggingReplicator;
@@ -49,6 +50,8 @@ class GAMEPLAYDEBUGGER_API UGameplayDebuggingControllerComponent : public UActor
 
 	float GetUpdateNavMeshTimeRemaining() const;
 
+	TWeakObjectPtr<class ADebugCameraController> GetDebugCameraController() { return DebugCameraController; }
+
 	FOnChangeEQSQuery OnNextEQSQuery;
 	FOnChangeEQSQuery OnPreviousEQSQuery;
 
@@ -63,6 +66,9 @@ protected:
 	UPROPERTY(Transient)
 	class UInputComponent* AIDebugViewInputComponent;
 
+	UPROPERTY(Transient)
+	class UInputComponent* DebugCameraInputComponent;
+
 	void OpenDebugTool();
 	void CloseDebugTool();
 
@@ -81,10 +87,12 @@ protected:
 	virtual void ToggleAIDebugView_SetView9();
 	virtual void NextEQSQuery();
 
-	virtual void BindAIDebugViewKeys();
+	virtual void BindAIDebugViewKeys(class UInputComponent*& InputComponent);
 	AGameplayDebuggingReplicator* GetDebuggingReplicator() const;
+	virtual void ToggleDebugCamera();
 
 	TWeakObjectPtr<APlayerController> PlayerOwner;
+	TWeakObjectPtr<ADebugCameraController> DebugCameraController;
 
 	/** Handle for efficient management of UpdateNavMesh timer */
 	FTimerHandle TimerHandle_UpdateNavMeshTimer;
@@ -96,3 +104,18 @@ protected:
 	uint32 bToolActivated : 1;
 	uint32 bWaitingForOwnersComponent : 1;
 };
+
+UCLASS()
+class AGaneplayDebuggerProxyHUD : public AHUD
+{
+	GENERATED_UCLASS_BODY()
+
+	FFontRenderInfo TextRenderInfo;
+
+	// Begin AActor Interface
+	virtual void PostRender() override;
+	// End AActor Interface
+
+	TWeakObjectPtr<AHUD> RedirectedHUD;
+};
+
diff --git a/Engine/Source/Developer/GameplayDebugger/Private/GameplayDebuggingComponent.cpp b/Engine/Source/Developer/GameplayDebugger/Private/GameplayDebuggingComponent.cpp
--- a/Engine/Source/Developer/GameplayDebugger/Private/GameplayDebuggingComponent.cpp
+++ b/Engine/Source/Developer/GameplayDebugger/Private/GameplayDebuggingComponent.cpp
@@ -284,7 +284,9 @@ void UGameplayDebuggingComponent::SelectTargetToDebug()
 {
 #if !(UE_BUILD_SHIPPING || UE_BUILD_TEST)
 	AGameplayDebuggingReplicator* Replicator = Cast<AGameplayDebuggingReplicator>(GetOwner());
-	APlayerController* MyPC = Replicator->GetLocalPlayerOwner();
+
+	UGameplayDebuggingControllerComponent* GDC = Replicator->FindComponentByClass<UGameplayDebuggingControllerComponent>();
+	APlayerController* MyPC = GDC && GDC->GetDebugCameraController().IsValid() ? GDC->GetDebugCameraController().Get() : Replicator->GetLocalPlayerOwner();
 
 	if (MyPC )
 	{
diff --git a/Engine/Source/Developer/GameplayDebugger/Private/GameplayDebuggingControllerComponent.cpp b/Engine/Source/Developer/GameplayDebugger/Private/GameplayDebuggingControllerComponent.cpp
--- a/Engine/Source/Developer/GameplayDebugger/Private/GameplayDebuggingControllerComponent.cpp
+++ b/Engine/Source/Developer/GameplayDebugger/Private/GameplayDebuggingControllerComponent.cpp
@@ -13,6 +13,19 @@ DEFINE_LOG_CATEGORY_STATIC(LogGameplayDebugging, Log, All);
 #define BUGIT_VIEWS (1<<EAIDebugDrawDataView::Basic) | (1 << EAIDebugDrawDataView::OverHead)
 #define BREAK_LINE_TEXT TEXT("________________________________________________________________")
 
+AGaneplayDebuggerProxyHUD::AGaneplayDebuggerProxyHUD(const FObjectInitializer& ObjectInitializer) : Super(ObjectInitializer)
+{
+}
+
+void AGaneplayDebuggerProxyHUD::PostRender()
+{
+	if (RedirectedHUD.IsValid())
+	{
+		RedirectedHUD->SetCanvas(Canvas, DebugCanvas);
+		RedirectedHUD->PostRender();
+	}
+}
+
 UGameplayDebuggingControllerComponent::UGameplayDebuggingControllerComponent(const FObjectInitializer& ObjectInitializer)
 	: Super(ObjectInitializer)
 	, KeyPressActivationTime(0.4f)
@@ -24,7 +37,7 @@ UGameplayDebuggingControllerComponent::UGameplayDebuggingControllerComponent(con
 	bTickInEditor=true;
 	PrimaryComponentTick.bStartWithTickEnabled = false;
 
-	DebugAITargetActor = NULL;
+	DebugAITargetActor = nullptr;
 	
 	bToolActivated = false;
 	bWaitingForOwnersComponent = false;
@@ -60,7 +73,7 @@ AGameplayDebuggingReplicator* UGameplayDebuggingControllerComponent::GetDebuggin
 void UGameplayDebuggingControllerComponent::BeginDestroy()
 {
 #if !(UE_BUILD_SHIPPING || UE_BUILD_TEST)
-	AHUD* GameHUD = PlayerOwner.IsValid() ? PlayerOwner->GetHUD() : NULL;
+	AHUD* GameHUD = PlayerOwner.IsValid() ? PlayerOwner->GetHUD() : nullptr;
 	if (GameHUD)
 	{
 		GameHUD->bShowHUD = true;
@@ -175,9 +188,14 @@ void UGameplayDebuggingControllerComponent::OnActivationKeyPressed()
 			Activate();
 			SetComponentTickEnabled(true);
 
-			BindAIDebugViewKeys();
+			BindAIDebugViewKeys(AIDebugViewInputComponent);
+			if (PlayerOwner.IsValid())
+			{
+				PlayerOwner->PushInputComponent(AIDebugViewInputComponent);
+			}
+
 			GetDebuggingReplicator()->EnableDraw(true);
-			GetDebuggingReplicator()->ServerReplicateMessage(NULL, EDebugComponentMessage::ActivateReplication, EAIDebugDrawDataView::Empty);
+			GetDebuggingReplicator()->ServerReplicateMessage(nullptr, EDebugComponentMessage::ActivateReplication, EAIDebugDrawDataView::Empty);
 		}
 
 		ControlKeyPressedTime = GetWorld()->GetTimeSeconds();
@@ -194,14 +212,14 @@ void UGameplayDebuggingControllerComponent::OnActivationKeyReleased()
 	EnableTargetSelection(false);
 	if (GetDebuggingReplicator() && bToolActivated)
 	{
-		if (KeyPressedTime < KeyPressActivationTime)
+		if (KeyPressedTime < KeyPressActivationTime && DebugCameraController.IsValid() == false)
 		{
 			CloseDebugTool();
 		}
 		else
 		{
-			APawn* TargetPawn = GetDebuggingReplicator()->GetDebugComponent() ? Cast<APawn>(GetDebuggingReplicator()->GetDebugComponent()->GetSelectedActor()) : NULL;
-			if (TargetPawn != NULL)
+			APawn* TargetPawn = GetDebuggingReplicator()->GetDebugComponent() ? Cast<APawn>(GetDebuggingReplicator()->GetDebugComponent()->GetSelectedActor()) : nullptr;
+			if (TargetPawn != nullptr)
 			{
 				FBehaviorTreeDelegates::OnDebugLocked.Broadcast(TargetPawn);
 			}
@@ -227,9 +245,9 @@ void UGameplayDebuggingControllerComponent::CloseDebugTool()
 	{
 		Deactivate();
 		SetComponentTickEnabled(false);
-		GetDebuggingReplicator()->ServerReplicateMessage(NULL, EDebugComponentMessage::DeactivateReplilcation, EAIDebugDrawDataView::Empty);
+		GetDebuggingReplicator()->ServerReplicateMessage(nullptr, EDebugComponentMessage::DeactivateReplilcation, EAIDebugDrawDataView::Empty);
 		GetDebuggingReplicator()->EnableDraw(false);
-		GetDebuggingReplicator()->ServerReplicateMessage(NULL, EDebugComponentMessage::DeactivateReplilcation, EAIDebugDrawDataView::Empty);
+		GetDebuggingReplicator()->ServerReplicateMessage(nullptr, EDebugComponentMessage::DeactivateReplilcation, EAIDebugDrawDataView::Empty);
 		bToolActivated = false;
 	}
 #endif //!(UE_BUILD_SHIPPING || UE_BUILD_TEST)
@@ -242,43 +260,43 @@ void UGameplayDebuggingControllerComponent::EnableTargetSelection(bool bEnable)
 #endif //!(UE_BUILD_SHIPPING || UE_BUILD_TEST)
 }
 
-void UGameplayDebuggingControllerComponent::BindAIDebugViewKeys()
+void UGameplayDebuggingControllerComponent::BindAIDebugViewKeys(class UInputComponent*& InputComponent)
 {
 #if !(UE_BUILD_SHIPPING || UE_BUILD_TEST)
-	if (!AIDebugViewInputComponent)
+	if (!InputComponent)
 	{
-		AIDebugViewInputComponent = NewObject<UInputComponent>(GetOwner(), TEXT("AIDebugViewInputComponent0"));
-		AIDebugViewInputComponent->RegisterComponent();
-
-		AIDebugViewInputComponent->BindKey(EKeys::NumPadZero, IE_Pressed, this, &UGameplayDebuggingControllerComponent::ToggleAIDebugView_SetView0);
-		AIDebugViewInputComponent->BindKey(EKeys::NumPadOne, IE_Pressed, this, &UGameplayDebuggingControllerComponent::ToggleAIDebugView_SetView1);
-		AIDebugViewInputComponent->BindKey(EKeys::NumPadTwo, IE_Pressed, this, &UGameplayDebuggingControllerComponent::ToggleAIDebugView_SetView2);
-		AIDebugViewInputComponent->BindKey(EKeys::NumPadThree, IE_Pressed, this, &UGameplayDebuggingControllerComponent::ToggleAIDebugView_SetView3);
-		AIDebugViewInputComponent->BindKey(EKeys::NumPadFour, IE_Pressed, this, &UGameplayDebuggingControllerComponent::ToggleAIDebugView_SetView4);
-		AIDebugViewInputComponent->BindKey(EKeys::NumPadFive, IE_Pressed, this, &UGameplayDebuggingControllerComponent::ToggleAIDebugView_SetView5);
-		AIDebugViewInputComponent->BindKey(EKeys::NumPadSix, IE_Pressed, this, &UGameplayDebuggingControllerComponent::ToggleAIDebugView_SetView6);
-		AIDebugViewInputComponent->BindKey(EKeys::NumPadSeven, IE_Pressed, this, &UGameplayDebuggingControllerComponent::ToggleAIDebugView_SetView7);
-		AIDebugViewInputComponent->BindKey(EKeys::NumPadEight, IE_Pressed, this, &UGameplayDebuggingControllerComponent::ToggleAIDebugView_SetView8);
-		AIDebugViewInputComponent->BindKey(EKeys::NumPadNine, IE_Pressed, this, &UGameplayDebuggingControllerComponent::ToggleAIDebugView_SetView9);
-		AIDebugViewInputComponent->BindKey(EKeys::Add, IE_Released, this, &UGameplayDebuggingControllerComponent::NextEQSQuery);
-#if USE_ALTERNATIVE_KEYS
-		AIDebugViewInputComponent->BindKey(FInputChord(EKeys::Zero, false, false, true, false), IE_Pressed, this, &UGameplayDebuggingControllerComponent::ToggleAIDebugView_SetView0);
-		AIDebugViewInputComponent->BindKey(FInputChord(EKeys::One, false, false, true, false), IE_Pressed, this, &UGameplayDebuggingControllerComponent::ToggleAIDebugView_SetView1);
-		AIDebugViewInputComponent->BindKey(FInputChord(EKeys::Two, false, false, true, false), IE_Pressed, this, &UGameplayDebuggingControllerComponent::ToggleAIDebugView_SetView2);
-		AIDebugViewInputComponent->BindKey(FInputChord(EKeys::Three, false, false, true, false), IE_Pressed, this, &UGameplayDebuggingControllerComponent::ToggleAIDebugView_SetView3);
-		AIDebugViewInputComponent->BindKey(FInputChord(EKeys::Four, false, false, true, false), IE_Pressed, this, &UGameplayDebuggingControllerComponent::ToggleAIDebugView_SetView4);
-		AIDebugViewInputComponent->BindKey(FInputChord(EKeys::Five, false, false, true, false), IE_Pressed, this, &UGameplayDebuggingControllerComponent::ToggleAIDebugView_SetView5);
-		AIDebugViewInputComponent->BindKey(FInputChord(EKeys::Six, false, false, true, false), IE_Pressed, this, &UGameplayDebuggingControllerComponent::ToggleAIDebugView_SetView6);
-		AIDebugViewInputComponent->BindKey(FInputChord(EKeys::Seven, false, false, true, false), IE_Pressed, this, &UGameplayDebuggingControllerComponent::ToggleAIDebugView_SetView7);
-		AIDebugViewInputComponent->BindKey(FInputChord(EKeys::Eight, false, false, true, false), IE_Pressed, this, &UGameplayDebuggingControllerComponent::ToggleAIDebugView_SetView8);
-		AIDebugViewInputComponent->BindKey(FInputChord(EKeys::Nine, false, false, true, false), IE_Pressed, this, &UGameplayDebuggingControllerComponent::ToggleAIDebugView_SetView9);
-		AIDebugViewInputComponent->BindKey(FInputChord(EKeys::Equals, false, false, true, false), IE_Released, this, &UGameplayDebuggingControllerComponent::NextEQSQuery);
-#endif
-	}
-	
-	if (PlayerOwner.IsValid())
-	{
-		PlayerOwner->PushInputComponent(AIDebugViewInputComponent);
+		InputComponent = NewObject<UInputComponent>(GetOwner(), TEXT("AIDebugViewInputComponent0"));
+		InputComponent->RegisterComponent();
+
+		if (UGameplayDebuggerSettings::StaticClass()->GetDefaultObject<UGameplayDebuggerSettings>()->UseAlternateKeys())
+		{
+			InputComponent->BindKey(FInputChord(EKeys::Zero, false, false, true, false), IE_Pressed, this, &UGameplayDebuggingControllerComponent::ToggleAIDebugView_SetView0);
+			InputComponent->BindKey(FInputChord(EKeys::One, false, false, true, false), IE_Pressed, this, &UGameplayDebuggingControllerComponent::ToggleAIDebugView_SetView1);
+			InputComponent->BindKey(FInputChord(EKeys::Two, false, false, true, false), IE_Pressed, this, &UGameplayDebuggingControllerComponent::ToggleAIDebugView_SetView2);
+			InputComponent->BindKey(FInputChord(EKeys::Three, false, false, true, false), IE_Pressed, this, &UGameplayDebuggingControllerComponent::ToggleAIDebugView_SetView3);
+			InputComponent->BindKey(FInputChord(EKeys::Four, false, false, true, false), IE_Pressed, this, &UGameplayDebuggingControllerComponent::ToggleAIDebugView_SetView4);
+			InputComponent->BindKey(FInputChord(EKeys::Five, false, false, true, false), IE_Pressed, this, &UGameplayDebuggingControllerComponent::ToggleAIDebugView_SetView5);
+			InputComponent->BindKey(FInputChord(EKeys::Six, false, false, true, false), IE_Pressed, this, &UGameplayDebuggingControllerComponent::ToggleAIDebugView_SetView6);
+			InputComponent->BindKey(FInputChord(EKeys::Seven, false, false, true, false), IE_Pressed, this, &UGameplayDebuggingControllerComponent::ToggleAIDebugView_SetView7);
+			InputComponent->BindKey(FInputChord(EKeys::Eight, false, false, true, false), IE_Pressed, this, &UGameplayDebuggingControllerComponent::ToggleAIDebugView_SetView8);
+			InputComponent->BindKey(FInputChord(EKeys::Nine, false, false, true, false), IE_Pressed, this, &UGameplayDebuggingControllerComponent::ToggleAIDebugView_SetView9);
+			InputComponent->BindKey(FInputChord(EKeys::Equals, false, false, true, false), IE_Released, this, &UGameplayDebuggingControllerComponent::NextEQSQuery);
+		}
+		else
+		{
+			InputComponent->BindKey(EKeys::NumPadZero, IE_Pressed, this, &UGameplayDebuggingControllerComponent::ToggleAIDebugView_SetView0);
+			InputComponent->BindKey(EKeys::NumPadOne, IE_Pressed, this, &UGameplayDebuggingControllerComponent::ToggleAIDebugView_SetView1);
+			InputComponent->BindKey(EKeys::NumPadTwo, IE_Pressed, this, &UGameplayDebuggingControllerComponent::ToggleAIDebugView_SetView2);
+			InputComponent->BindKey(EKeys::NumPadThree, IE_Pressed, this, &UGameplayDebuggingControllerComponent::ToggleAIDebugView_SetView3);
+			InputComponent->BindKey(EKeys::NumPadFour, IE_Pressed, this, &UGameplayDebuggingControllerComponent::ToggleAIDebugView_SetView4);
+			InputComponent->BindKey(EKeys::NumPadFive, IE_Pressed, this, &UGameplayDebuggingControllerComponent::ToggleAIDebugView_SetView5);
+			InputComponent->BindKey(EKeys::NumPadSix, IE_Pressed, this, &UGameplayDebuggingControllerComponent::ToggleAIDebugView_SetView6);
+			InputComponent->BindKey(EKeys::NumPadSeven, IE_Pressed, this, &UGameplayDebuggingControllerComponent::ToggleAIDebugView_SetView7);
+			InputComponent->BindKey(EKeys::NumPadEight, IE_Pressed, this, &UGameplayDebuggingControllerComponent::ToggleAIDebugView_SetView8);
+			InputComponent->BindKey(EKeys::NumPadNine, IE_Pressed, this, &UGameplayDebuggingControllerComponent::ToggleAIDebugView_SetView9);
+			InputComponent->BindKey(EKeys::Add, IE_Released, this, &UGameplayDebuggingControllerComponent::NextEQSQuery);
+		}
+		InputComponent->BindKey(EKeys::Tab, IE_Released, this, &UGameplayDebuggingControllerComponent::ToggleDebugCamera);
 	}
 #endif //!(UE_BUILD_SHIPPING || UE_BUILD_TEST)
 }
@@ -295,11 +313,11 @@ bool UGameplayDebuggingControllerComponent::CanToggleView()
 void UGameplayDebuggingControllerComponent::ToggleAIDebugView_SetView0()
 {
 #if !(UE_BUILD_SHIPPING || UE_BUILD_TEST)
-	APawn* Pawn = PlayerOwner.IsValid() ? PlayerOwner->GetPawnOrSpectator() : NULL;
+	APawn* Pawn = PlayerOwner.IsValid() ? PlayerOwner->GetPawnOrSpectator() : nullptr;
 
 	if (PlayerOwner.IsValid() && Pawn && bToolActivated && GetDebuggingReplicator())
 	{
-		if (UGameplayDebuggingComponent* OwnerComp = GetDebuggingReplicator() ? GetDebuggingReplicator()->GetDebugComponent() : NULL)
+		if (UGameplayDebuggingComponent* OwnerComp = GetDebuggingReplicator() ? GetDebuggingReplicator()->GetDebugComponent() : nullptr)
 		{
 			GameplayDebuggerSettings(GetDebuggingReplicator()).DebuggerShowFlags ^= 1 << EAIDebugDrawDataView::NavMesh;
 
@@ -390,8 +408,8 @@ float UGameplayDebuggingControllerComponent::GetUpdateNavMeshTimeRemaining() con
 void UGameplayDebuggingControllerComponent::UpdateNavMeshTimer()
 {
 #if !(UE_BUILD_SHIPPING || UE_BUILD_TEST)
-	const APawn* PlayerPawn = PlayerOwner.IsValid() ? PlayerOwner->GetPawnOrSpectator() : NULL;
-	UGameplayDebuggingComponent* DebuggingComponent = GetDebuggingReplicator() ? GetDebuggingReplicator()->GetDebugComponent() : NULL;
+	const APawn* PlayerPawn = PlayerOwner.IsValid() ? PlayerOwner->GetPawnOrSpectator() : nullptr;
+	UGameplayDebuggingComponent* DebuggingComponent = GetDebuggingReplicator() ? GetDebuggingReplicator()->GetDebugComponent() : nullptr;
 	if (DebuggingComponent)
 	{
 		const AActor* SelectedActor = DebuggingComponent->GetSelectedActor();
@@ -410,3 +428,71 @@ void UGameplayDebuggingControllerComponent::UpdateNavMeshTimer()
 	}
 #endif //!(UE_BUILD_SHIPPING || UE_BUILD_TEST)
 }
+
+void UGameplayDebuggingControllerComponent::ToggleDebugCamera()
+{
+#if !(UE_BUILD_SHIPPING || UE_BUILD_TEST)
+	if(!bToolActivated || PlayerOwner == nullptr)
+	{
+		return;
+	}
+
+	if (DebugCameraController.IsValid() == false)
+	{
+		// spawn if necessary
+		// and ungly @HACK to be able to spawn camera in game world rather then
+		// in editor world (if running PIE). Hate it, but it works, and 
+		// this is a debugging tool		
+		FActorSpawnParameters SpawnInfo;
+		SpawnInfo.bNoCollisionFail = true;
+		SpawnInfo.Owner = PlayerOwner->GetWorldSettings();
+		SpawnInfo.Instigator = PlayerOwner->Instigator;
+		DebugCameraController = GetWorld()->SpawnActor<ADebugCameraController>(SpawnInfo);
+
+		if (DebugCameraController.IsValid())
+		{
+			// set up new controller
+			DebugCameraController->OnActivate(PlayerOwner.Get());
+
+			// then switch to it
+			PlayerOwner->Player->SwitchController(DebugCameraController.Get());
+
+			FActorSpawnParameters SpawnInfo;
+			SpawnInfo.Owner = PlayerOwner.Get();
+			SpawnInfo.Instigator = PlayerOwner->Instigator;
+			SpawnInfo.bNoCollisionFail = true;
+			SpawnInfo.ObjectFlags |= RF_Transient;	// We never want these to save into a map
+			AGaneplayDebuggerProxyHUD* ProxyHUD = GetWorld()->SpawnActor<AGaneplayDebuggerProxyHUD>(SpawnInfo);
+			ProxyHUD->RedirectedHUD = PlayerOwner->MyHUD;
+			DebugCameraController->MyHUD = ProxyHUD;
+			BindAIDebugViewKeys(DebugCameraInputComponent);
+			DebugCameraInputComponent->BindKey(ActivationKey, IE_Pressed, this, &UGameplayDebuggingControllerComponent::OnActivationKeyPressed);
+			DebugCameraInputComponent->BindKey(ActivationKey, IE_Released, this, &UGameplayDebuggingControllerComponent::OnActivationKeyReleased);
+			DebugCameraController->PushInputComponent(DebugCameraInputComponent);
+
+			DebugCameraController->ChangeState(NAME_Default);
+			DebugCameraController->ChangeState(NAME_Spectating);
+		}
+	}
+	else
+	{
+		DebugCameraController->PopInputComponent(DebugCameraInputComponent);
+		DebugCameraInputComponent = nullptr;
+		DebugCameraController->OriginalPlayer->SwitchController(DebugCameraController->OriginalControllerRef);
+		DebugCameraController->OnDeactivate(DebugCameraController->OriginalControllerRef);
+		GetWorld()->DestroyActor(DebugCameraController.Get(), false, false);
+		DebugCameraController = nullptr;
+
+		if(AIDebugViewInputComponent == nullptr)
+		{
+			BindAIDebugViewKeys(AIDebugViewInputComponent);
+		}
+		if (PlayerOwner.IsValid())
+		{
+			PlayerOwner->PushInputComponent(AIDebugViewInputComponent);
+		}
+
+	}
+
+#endif //!(UE_BUILD_SHIPPING || UE_BUILD_TEST)
+	}
diff --git a/Engine/Source/Developer/GameplayDebugger/Private/GameplayDebuggingHUDComponent.cpp b/Engine/Source/Developer/GameplayDebugger/Private/GameplayDebuggingHUDComponent.cpp
--- a/Engine/Source/Developer/GameplayDebugger/Private/GameplayDebuggingHUDComponent.cpp
+++ b/Engine/Source/Developer/GameplayDebugger/Private/GameplayDebuggingHUDComponent.cpp
@@ -156,7 +156,11 @@ void AGameplayDebuggingHUDComponent::DrawMenu(const float X, const float Y, clas
 		}
 
 		const FString KeyDesc = ActivationKeyName != ActivationKeyDisplayName ? FString::Printf(TEXT("(%s key)"), *ActivationKeyName) : TEXT("");
-		const FString HeaderDesc = FString::Printf(TEXT("Tap %s %s to close, use Numpad to toggle categories"), *ActivationKeyDisplayName, *KeyDesc);
+		FString HeaderDesc = FString::Printf(TEXT("Tap %s %s to close, use Numpad numbers to toggle categories"), *ActivationKeyDisplayName, *KeyDesc);
+		if (UGameplayDebuggerSettings::StaticClass()->GetDefaultObject<UGameplayDebuggerSettings>()->UseAlternateKeys())
+		{
+			HeaderDesc = FString::Printf(TEXT("Tap %s %s to close, use Alt + number to toggle categories"), *ActivationKeyDisplayName, *KeyDesc);
+		}
 
 		float HeaderWidth = 0.0f;
 		CalulateStringSize(DefaultContext, DefaultContext.Font, HeaderDesc, HeaderWidth, MaxHeight);
@@ -171,6 +175,16 @@ void AGameplayDebuggingHUDComponent::DrawMenu(const float X, const float Y, clas
 			MaxHeight = FMath::Max(MaxHeight, StrHeight);
 		}
 
+		{
+			const int32 DebugCameraIndex = Categories.Add(FDebugCategoryView());
+			CategoriesWidth.AddZeroed(1);
+			Categories[DebugCameraIndex].Desc = FString::Printf(TEXT(" %s[Tab]: %s  "), GDC && GDC->GetDebugCameraController().IsValid() ? TEXT("{Green}") : TEXT("{White}"), TEXT("Debug Camera"));
+			float StrHeight = 0.0f;
+			CalulateStringSize(DefaultContext, DefaultContext.Font, Categories[DebugCameraIndex].Desc, CategoriesWidth[DebugCameraIndex], StrHeight);
+			TotalWidth += CategoriesWidth[DebugCameraIndex];
+			MaxHeight = FMath::Max(MaxHeight, StrHeight);
+		}
+
 		TotalWidth = FMath::Max(TotalWidth, HeaderWidth);
 
 		FCanvasTileItem TileItem(FVector2D(10, 10), GWhiteTexture, FVector2D(TotalWidth + 20, MaxHeight + 20), FColor(0, 0, 0, 20));
@@ -179,15 +193,15 @@ void AGameplayDebuggingHUDComponent::DrawMenu(const float X, const float Y, clas
 
 		PrintString(DefaultContext, FColorList::LightBlue, HeaderDesc, MenuStartX + 2.f, MenuStartY + 2.f);
 
-			float XPos = MenuStartX + 20.f;
-			for (int32 i = 0; i < Categories.Num(); i++)
-			{
-				const bool bIsActive = GameplayDebuggerSettings(GetDebuggingReplicator()).CheckFlag(Categories[i].View) ? true : false;
-				const bool bIsDisabled = Categories[i].View == EAIDebugDrawDataView::NavMesh ? false : (DebugComponent && DebugComponent->GetSelectedActor() ? false: true);
+		float XPos = MenuStartX + 20.f;
+		for (int32 i = 0; i < Categories.Num(); i++)
+		{
+			const bool bIsActive = GameplayDebuggerSettings(GetDebuggingReplicator()).CheckFlag(Categories[i].View) ? true : false;
+			const bool bIsDisabled = Categories[i].View == EAIDebugDrawDataView::NavMesh ? false : (DebugComponent && DebugComponent->GetSelectedActor() ? false: true);
 
-				PrintString(DefaultContext, bIsDisabled ? (bIsActive ? FColorList::DarkGreen  : FColorList::LightGrey) : (bIsActive ? FColorList::Green : FColorList::White), Categories[i].Desc, XPos, MenuStartY + MaxHeight + 2.f);
-				XPos += CategoriesWidth[i];
-			}
+			PrintString(DefaultContext, bIsDisabled ? (bIsActive ? FColorList::DarkGreen  : FColorList::LightGrey) : (bIsActive ? FColorList::Green : FColorList::White), Categories[i].Desc, XPos, MenuStartY + MaxHeight + 2.f);
+			XPos += CategoriesWidth[i];
+		}
 		DefaultContext.Font = OldFont;
 	}
 
@@ -209,6 +223,7 @@ void AGameplayDebuggingHUDComponent::DrawDebugComponentData(APlayerController* M
 	const FVector ScreenLoc = SelectedActor ? ProjectLocation(DefaultContext, SelectedActor->GetActorLocation() + FVector(0.f, 0.f, SelectedActor->GetSimpleCollisionHalfHeight())) : FVector::ZeroVector;
 
 	OverHeadContext = FPrintContext(GEngine->GetSmallFont(), Canvas, ScreenLoc.X, ScreenLoc.Y);
+	DefaultContext.CursorY += 20;
 
 	FGameplayDebuggerSettings DebuggerSettings = GameplayDebuggerSettings(GetDebuggingReplicator());
 	bool bForceOverhead = false;
diff --git a/Engine/Source/Developer/GameplayDebugger/Private/GameplayDebuggingReplicator.cpp b/Engine/Source/Developer/GameplayDebugger/Private/GameplayDebuggingReplicator.cpp
--- a/Engine/Source/Developer/GameplayDebugger/Private/GameplayDebuggingReplicator.cpp
+++ b/Engine/Source/Developer/GameplayDebugger/Private/GameplayDebuggingReplicator.cpp
@@ -655,10 +655,17 @@ void AGameplayDebuggingReplicator::DrawDebugData(class UCanvas* Canvas, class AP
 		return;
 	}
 
-	const bool bAllowToDraw = Canvas && Canvas->SceneView && (Canvas->SceneView->ViewActor == LocalPlayerOwner->AcknowledgedPawn || Canvas->SceneView->ViewActor == LocalPlayerOwner->GetPawnOrSpectator());
+	bool bAllowToDraw = Canvas && Canvas->SceneView && (Canvas->SceneView->ViewActor == LocalPlayerOwner->AcknowledgedPawn || Canvas->SceneView->ViewActor == LocalPlayerOwner->GetPawnOrSpectator());
 	if (!bAllowToDraw)
 	{
-		return;
+		// check for spectator debug camera
+		UGameplayDebuggingControllerComponent* GDC = FindComponentByClass<UGameplayDebuggingControllerComponent>();
+		bAllowToDraw = GDC && GDC->GetDebugCameraController().IsValid() && Canvas->SceneView->ViewActor->GetInstigatorController() == GDC->GetDebugCameraController().Get();
+		
+		if (!bAllowToDraw)
+		{
+			return;
+		}
 	}
 
 	if (!DebugRenderer.IsValid() && DebugComponentHUDClass.IsValid())
