Fixed monolithic build issues for maxos.

diff --git a/Engine/Source/Developer/GameplayDebugger/Classes/GameplayDebuggingControllerComponent.h b/Engine/Source/Developer/GameplayDebugger/Classes/GameplayDebuggingControllerComponent.h
--- a/Engine/Source/Developer/GameplayDebugger/Classes/GameplayDebuggingControllerComponent.h
+++ b/Engine/Source/Developer/GameplayDebugger/Classes/GameplayDebuggingControllerComponent.h
@@ -9,6 +9,7 @@
 #include "GameplayDebuggingTypes.h"
 #include "GameplayDebugger.h"
 #include "GameFramework/HUD.h"
+#include "Engine/DebugCameraController.h"
 #include "GameplayDebuggingControllerComponent.generated.h"
 
 class AGameplayDebuggingReplicator;
@@ -17,7 +18,7 @@ UCLASS(config = Engine)
 class GAMEPLAYDEBUGGER_API UGameplayDebuggingControllerComponent : public UActorComponent
 {
 	GENERATED_UCLASS_BODY()
-
+		
 #if WITH_HOT_RELOAD_CTORS
 	/** DO NOT USE. This constructor is for internal usage only for hot-reload purposes. */
 	UGameplayDebuggingControllerComponent(FVTableHelper& Helper);
@@ -50,7 +51,7 @@ class GAMEPLAYDEBUGGER_API UGameplayDebuggingControllerComponent : public UActor
 
 	float GetUpdateNavMeshTimeRemaining() const;
 
-	TWeakObjectPtr<class ADebugCameraController> GetDebugCameraController() { return DebugCameraController; }
+	TWeakObjectPtr<ADebugCameraController> GetDebugCameraController() { return DebugCameraController; }
 
 	FOnChangeEQSQuery OnNextEQSQuery;
 	FOnChangeEQSQuery OnPreviousEQSQuery;
