Fix for osx compilation issues.

diff --git a/Engine/Source/Developer/GameplayDebugger/Private/GameplayDebuggingControllerComponent.cpp b/Engine/Source/Developer/GameplayDebugger/Private/GameplayDebuggingControllerComponent.cpp
--- a/Engine/Source/Developer/GameplayDebugger/Private/GameplayDebuggingControllerComponent.cpp
+++ b/Engine/Source/Developer/GameplayDebugger/Private/GameplayDebuggingControllerComponent.cpp
@@ -6,6 +6,7 @@
 #include "BehaviorTreeDelegates.h"
 #include "TimerManager.h"
 #include "VisualLogger/VisualLogger.h"
+#include "GameplayDebuggerSettings.h"
 
 DEFINE_LOG_CATEGORY_STATIC(LogGameplayDebugging, Log, All);
 
