Fix PS4 mono compile.

diff --git a/Engine/Source/Developer/GameplayDebugger/Private/GameplayDebuggingControllerComponent.cpp b/Engine/Source/Developer/GameplayDebugger/Private/GameplayDebuggingControllerComponent.cpp
--- a/Engine/Source/Developer/GameplayDebugger/Private/GameplayDebuggingControllerComponent.cpp
+++ b/Engine/Source/Developer/GameplayDebugger/Private/GameplayDebuggingControllerComponent.cpp
@@ -7,6 +7,7 @@
 #include "TimerManager.h"
 #include "VisualLogger/VisualLogger.h"
 #include "GameplayDebuggerSettings.h"
+#include "Player.h"
 
 DEFINE_LOG_CATEGORY_STATIC(LogGameplayDebugging, Log, All);
 
