GameplayAbilities module is no longer loaded by default [UE-14860] in all games. Can be optionally loaded by games that need it.

diff --git a/Engine/Source/Developer/GameplayDebugger/Private/GameplayDebuggingComponent.cpp b/Engine/Source/Developer/GameplayDebugger/Private/GameplayDebuggingComponent.cpp
--- a/Engine/Source/Developer/GameplayDebugger/Private/GameplayDebuggingComponent.cpp
+++ b/Engine/Source/Developer/GameplayDebugger/Private/GameplayDebuggingComponent.cpp
@@ -20,6 +20,7 @@
 #include "GameFramework/CharacterMovementComponent.h"
 #include "Engine/Channel.h"
 #include "Animation/AnimMontage.h"
+#include "GameplayAbilitiesModule.h"
 
 #if WITH_EDITOR
 #include "Editor/EditorEngine.h"
@@ -536,35 +537,16 @@ void UGameplayDebuggingComponent::CollectBasicBehaviorData(APawn* MyPawn)
 void UGameplayDebuggingComponent::CollectBasicAbilityData(APawn* MyPawn)
 {
 #if !(UE_BUILD_SHIPPING || UE_BUILD_TEST)
-	UAbilitySystemComponent* AbilityComp = MyPawn->FindComponentByClass<UAbilitySystemComponent>();
-	bIsUsingAbilities = (AbilityComp != nullptr);
-
-	int32 NumActive = 0;
-	if (AbilityComp)
+	if (IGameplayAbilitiesModule::IsAvailable())
 	{
-		AbilityInfo = TEXT("");
-
-		for (const FGameplayAbilitySpec& AbilitySpec : AbilityComp->GetActivatableAbilities())
-		{
-			if (AbilitySpec.Ability && AbilitySpec.IsActive())
-			{
-				if (NumActive)
-				{
-					AbilityInfo += TEXT(", ");
-				}
-
-				UClass* AbilityClass = AbilitySpec.Ability->GetClass();
-				FString AbClassName = GetNameSafe(AbilityClass);
-				AbClassName.RemoveFromEnd(TEXT("_c"));
-
-				AbilityInfo += AbClassName;
-				NumActive++;
-			}
-		}
+		bool bUsingAbilities;
+		IGameplayAbilitiesModule& AbilitiesModule = FModuleManager::LoadModuleChecked<IGameplayAbilitiesModule>("GameplayAbilities");
+		AbilitiesModule.GetActiveAbilitiesDebugDataForActor(MyPawn, AbilityInfo, bUsingAbilities);
+		bIsUsingAbilities = bUsingAbilities;
 	}
-
-	if (NumActive == 0)
+	else
 	{
+		bIsUsingAbilities = false;
 		AbilityInfo = TEXT("None");
 	}
 #endif //!(UE_BUILD_SHIPPING || UE_BUILD_TEST)
diff --git a/Engine/Source/Runtime/GameplayAbilities/Private/GameplayAbilitiesExec.cpp b/Engine/Source/Runtime/GameplayAbilities/Private/GameplayAbilitiesExec.cpp
--- a/Engine/Source/Runtime/GameplayAbilities/Private/GameplayAbilitiesExec.cpp
+++ b/Engine/Source/Runtime/GameplayAbilities/Private/GameplayAbilitiesExec.cpp
@@ -25,16 +25,17 @@ bool FGameplayAbilitiesExec::Exec(UWorld* Inworld, const TCHAR* Cmd, FOutputDevi
 		return false;
 	}
 
-	UAbilitySystemGlobals& AbilitySystemGlobals = UAbilitySystemGlobals::Get();
 	bool bHandled = false;
 
 	if (FParse::Command(&Cmd, TEXT("ToggleIgnoreAbilitySystemCooldowns")))
 	{
+		UAbilitySystemGlobals& AbilitySystemGlobals = UAbilitySystemGlobals::Get();
 		AbilitySystemGlobals.ToggleIgnoreAbilitySystemCooldowns();
 		bHandled = true;
 	}
 	else if (FParse::Command(&Cmd, TEXT("ToggleIgnoreAbilitySystemCosts")))
 	{
+		UAbilitySystemGlobals& AbilitySystemGlobals = UAbilitySystemGlobals::Get();
 		AbilitySystemGlobals.ToggleIgnoreAbilitySystemCosts();
 		bHandled = true;
 	}
diff --git a/Engine/Source/Runtime/GameplayAbilities/Private/GameplayAbilitiesModule.cpp b/Engine/Source/Runtime/GameplayAbilities/Private/GameplayAbilitiesModule.cpp
--- a/Engine/Source/Runtime/GameplayAbilities/Private/GameplayAbilitiesModule.cpp
+++ b/Engine/Source/Runtime/GameplayAbilities/Private/GameplayAbilitiesModule.cpp
@@ -31,6 +31,8 @@ class FGameplayAbilitiesModule : public IGameplayAbilitiesModule
 
 	UAbilitySystemGlobals *AbilitySystemGlobals;
 
+	void GetActiveAbilitiesDebugDataForActor(AActor* Actor, FString& AbilityString, bool& bIsUsingAbilities);
+
 private:
 	
 };
@@ -50,3 +52,38 @@ void FGameplayAbilitiesModule::ShutdownModule()
 
 	AbilitySystemGlobals = NULL;
 }
+
+void FGameplayAbilitiesModule::GetActiveAbilitiesDebugDataForActor(AActor* Actor, FString& AbilityString, bool& bIsUsingAbilities)
+{
+	UAbilitySystemComponent* AbilityComp = Actor->FindComponentByClass<UAbilitySystemComponent>();
+	bIsUsingAbilities = (AbilityComp != nullptr);
+	  
+	int32 NumActive = 0;
+	if (AbilityComp)
+	{
+		AbilityString = TEXT("");
+	  
+		for (const FGameplayAbilitySpec& AbilitySpec : AbilityComp->GetActivatableAbilities())
+		{
+	  		if (AbilitySpec.Ability && AbilitySpec.IsActive())
+	  		{
+	  			if (NumActive)
+	  			{
+					AbilityString += TEXT(", ");
+	  			}
+	  
+	  			UClass* AbilityClass = AbilitySpec.Ability->GetClass();
+	  			FString AbClassName = GetNameSafe(AbilityClass);
+	  			AbClassName.RemoveFromEnd(TEXT("_c"));
+	  
+				AbilityString += AbClassName;
+	  			NumActive++;
+	  		}
+		}
+	}
+	  
+	if (NumActive == 0)
+	{
+		AbilityString = TEXT("None");
+	}
+}
diff --git a/Engine/Source/Runtime/GameplayAbilities/Public/GameplayAbilitiesModule.h b/Engine/Source/Runtime/GameplayAbilities/Public/GameplayAbilitiesModule.h
--- a/Engine/Source/Runtime/GameplayAbilities/Public/GameplayAbilitiesModule.h
+++ b/Engine/Source/Runtime/GameplayAbilities/Public/GameplayAbilitiesModule.h
@@ -36,5 +36,7 @@ public:
 	}
 
 	virtual UAbilitySystemGlobals* GetAbilitySystemGlobals() = 0;
+
+	virtual void GetActiveAbilitiesDebugDataForActor(AActor* Actor, FString& AbilityString, bool& bIsUsingAbilities)=0;
 };
 
