Fixed issue UE-14330 "Opening Gameplay Debugger while in PIE removes print string output from screen in subsequent PIE modes"
Fixed issues with Gameplay Debugger's input - it was breaking project related inputs sometimes.

diff --git a/Engine/Source/Developer/GameplayDebugger/Classes/GameplayDebuggingReplicator.h b/Engine/Source/Developer/GameplayDebugger/Classes/GameplayDebuggingReplicator.h
--- a/Engine/Source/Developer/GameplayDebugger/Classes/GameplayDebuggingReplicator.h
+++ b/Engine/Source/Developer/GameplayDebugger/Classes/GameplayDebuggingReplicator.h
@@ -117,6 +117,8 @@ class GAMEPLAYDEBUGGER_API AGameplayDebuggingReplicator : public AActor
 
 	virtual void TickActor(float DeltaTime, enum ELevelTick TickType, FActorTickFunction& ThisTickFunction) override;
 
+	virtual void BeginDestroy() override;
+
 	UGameplayDebuggingComponent* GetDebugComponent();
 
 	bool IsToolCreated();
diff --git a/Engine/Source/Developer/GameplayDebugger/Private/GameplayDebuggingControllerComponent.cpp b/Engine/Source/Developer/GameplayDebugger/Private/GameplayDebuggingControllerComponent.cpp
--- a/Engine/Source/Developer/GameplayDebugger/Private/GameplayDebuggingControllerComponent.cpp
+++ b/Engine/Source/Developer/GameplayDebugger/Private/GameplayDebuggingControllerComponent.cpp
@@ -250,6 +250,12 @@ void UGameplayDebuggingControllerComponent::CloseDebugTool()
 		GetDebuggingReplicator()->ServerReplicateMessage(nullptr, EDebugComponentMessage::DeactivateReplilcation, EAIDebugDrawDataView::Empty);
 		GetDebuggingReplicator()->EnableDraw(false);
 		GetDebuggingReplicator()->ServerReplicateMessage(nullptr, EDebugComponentMessage::DeactivateReplilcation, EAIDebugDrawDataView::Empty);
+
+		if (PlayerOwner.IsValid())
+		{
+			PlayerOwner->PopInputComponent(AIDebugViewInputComponent);
+		}
+		AIDebugViewInputComponent = nullptr;
 		bToolActivated = false;
 	}
 #endif //!(UE_BUILD_SHIPPING || UE_BUILD_TEST)
@@ -298,7 +304,7 @@ void UGameplayDebuggingControllerComponent::BindAIDebugViewKeys(class UInputComp
 			InputComponent->BindKey(EKeys::NumPadNine, IE_Pressed, this, &UGameplayDebuggingControllerComponent::ToggleAIDebugView_SetView9);
 			InputComponent->BindKey(EKeys::Add, IE_Released, this, &UGameplayDebuggingControllerComponent::NextEQSQuery);
 		}
-		InputComponent->BindKey(EKeys::Tab, IE_Released, this, &UGameplayDebuggingControllerComponent::ToggleDebugCamera);
+		InputComponent->BindKey(FInputChord(EKeys::Tab, false, false, false, false), IE_Released, this, &UGameplayDebuggingControllerComponent::ToggleDebugCamera);
 	}
 #endif //!(UE_BUILD_SHIPPING || UE_BUILD_TEST)
 }
@@ -490,10 +496,10 @@ void UGameplayDebuggingControllerComponent::ToggleDebugCamera()
 		if(AIDebugViewInputComponent == nullptr)
 		{
 			BindAIDebugViewKeys(AIDebugViewInputComponent);
-		}
-		if (PlayerOwner.IsValid())
-		{
-			PlayerOwner->PushInputComponent(AIDebugViewInputComponent);
+			if (PlayerOwner.IsValid())
+			{
+				PlayerOwner->PushInputComponent(AIDebugViewInputComponent);
+			}
 		}
 
 	}
diff --git a/Engine/Source/Developer/GameplayDebugger/Private/GameplayDebuggingReplicator.cpp b/Engine/Source/Developer/GameplayDebugger/Private/GameplayDebuggingReplicator.cpp
--- a/Engine/Source/Developer/GameplayDebugger/Private/GameplayDebuggingReplicator.cpp
+++ b/Engine/Source/Developer/GameplayDebugger/Private/GameplayDebuggingReplicator.cpp
@@ -276,6 +276,19 @@ void AGameplayDebuggingReplicator::BeginPlay()
 #endif //!(UE_BUILD_SHIPPING || UE_BUILD_TEST)
 }
 
+
+void AGameplayDebuggingReplicator::BeginDestroy()
+{
+#if !(UE_BUILD_SHIPPING || UE_BUILD_TEST)
+	GEngine->bEnableOnScreenDebugMessages = true;
+	if (IsDrawEnabled())
+	{
+		EnableDraw(false);
+	}
+#endif
+	Super::BeginDestroy();
+}
+
 void AGameplayDebuggingReplicator::PostNetInit()
 {
 	Super::PostNetInit();
