Fixed broken Gameplay Debugger's debug camera binding after debug camera usage.

diff --git a/Engine/Source/Developer/GameplayDebugger/Private/GameplayDebuggingControllerComponent.cpp b/Engine/Source/Developer/GameplayDebugger/Private/GameplayDebuggingControllerComponent.cpp
--- a/Engine/Source/Developer/GameplayDebugger/Private/GameplayDebuggingControllerComponent.cpp
+++ b/Engine/Source/Developer/GameplayDebugger/Private/GameplayDebuggingControllerComponent.cpp
@@ -482,6 +482,12 @@ void UGameplayDebuggingControllerComponent::ToggleDebugCamera()
 
 			DebugCameraController->ChangeState(NAME_Default);
 			DebugCameraController->ChangeState(NAME_Spectating);
+
+			if (PlayerOwner.IsValid())
+			{
+				PlayerOwner->PopInputComponent(AIDebugViewInputComponent);
+			}
+			AIDebugViewInputComponent = nullptr;
 		}
 	}
 	else
