Fixed Gameplay Debugger for configuration with dedicated server (client-server environment).

diff --git a/Engine/Source/Developer/GameplayDebugger/Private/GameplayDebugger.cpp b/Engine/Source/Developer/GameplayDebugger/Private/GameplayDebugger.cpp
--- a/Engine/Source/Developer/GameplayDebugger/Private/GameplayDebugger.cpp
+++ b/Engine/Source/Developer/GameplayDebugger/Private/GameplayDebugger.cpp
@@ -434,89 +434,178 @@ void FGameplayDebugger::OnLevelActorDeleted(AActor* InActor)
 }
 #endif //WITH_EDITOR
 
-bool FGameplayDebugger::Exec(UWorld* Inworld, const TCHAR* Cmd, FOutputDevice& Ar)
+bool FGameplayDebugger::Exec(UWorld* InWorld, const TCHAR* Cmd, FOutputDevice& Ar)
 {
 #if !(UE_BUILD_SHIPPING || UE_BUILD_TEST)
 	bool bHandled = false;
 
-	
-	if (!Inworld)
+	if (FParse::Command(&Cmd, TEXT("RunEQS")) && InWorld)
 	{
-		return bHandled;
-	}
+		APlayerController* MyPC = InWorld->GetFirstPlayerController();
+		UAISystem* AISys = UAISystem::GetCurrent(*InWorld);
 
-	if (FParse::Command(&Cmd, TEXT("EnableGDT")))
-	{
-		if (Inworld->GetNetMode() == NM_Client)
+		UEnvQueryManager* EQS = AISys ? AISys->GetEnvironmentQueryManager() : NULL;
+		if (MyPC && EQS)
 		{
-			APlayerController* LocalPC = NULL;
-			for (FConstPlayerControllerIterator Iterator = Inworld->GetPlayerControllerIterator(); Iterator; ++Iterator)
+			AGameplayDebuggingReplicator* DebuggingReplicator = NULL;
+			for (TActorIterator<AGameplayDebuggingReplicator> It(InWorld); It; ++It)
 			{
-				APlayerController* PlayerController = *Iterator;
-				if (PlayerController && PlayerController->IsLocalPlayerController())
+				AGameplayDebuggingReplicator* A = *It;
+				if (!A->IsPendingKill())
 				{
-					LocalPC = PlayerController;
-					break;
+					DebuggingReplicator = A;
+					if (!DebuggingReplicator->IsGlobalInWorld() && DebuggingReplicator->GetLocalPlayerOwner() == MyPC)
+					{
+						break;
+					}
 				}
 			}
 
-			if (LocalPC)
+			UObject* Target = DebuggingReplicator != NULL ? DebuggingReplicator->GetSelectedActorToDebug() : NULL;
+			FString QueryName = FParse::Token(Cmd, 0);
+			if (Target)
+			{
+				AISys->RunEQS(QueryName, Target);
+			}
+			else
+			{
+				MyPC->ClientMessage(TEXT("No debugging target to run EQS"));
+			}
+		}
+
+		return true;
+	}
+
+
+	if (FParse::Command(&Cmd, TEXT("EnableGDT")) == false)
+	{
+		return false;
+	}
+
+	FString UniquePlayerId = FParse::Token(Cmd, 0);
+	APlayerController* LocalPC = NULL;
+	UWorld* MyWorld = InWorld;
+
+	if (MyWorld == nullptr)
+	{
+		if (UniquePlayerId.Len() > 0)
+		{
+			// let's find correct world based on Player Id
+			const TIndirectArray<FWorldContext> WorldContexts = GEngine->GetWorldContexts();
+			for (auto& Context : WorldContexts)
 			{
-				AGameplayDebuggingReplicator* Replicator = NULL;
-				for (TActorIterator<AGameplayDebuggingReplicator> It(Inworld); It; ++It)
+				if (Context.WorldType != EWorldType::Game && Context.WorldType != EWorldType::PIE)
 				{
-					Replicator = *It;
-					if (Replicator && !Replicator->IsPendingKill())
+					continue;
+				}
+
+				UWorld *CurrentWorld = Context.World();
+				for (FConstPlayerControllerIterator Iterator = CurrentWorld->GetPlayerControllerIterator(); Iterator; ++Iterator)
+				{
+					APlayerController* PC = *Iterator;
+					if (PC && PC->PlayerState->UniqueId.ToString() == UniquePlayerId)
 					{
-						APlayerController* PCOwner = Replicator->GetLocalPlayerOwner();
-						if (LocalPC == PCOwner)
-						{
-							break;
-						}
+						LocalPC = PC;
+						MyWorld = PC->GetWorld();
+						break;
 					}
-					Replicator = NULL;
 				}
 
-				if (!Replicator)
-				{					
-					LocalPC->ClientMessage(TEXT("Enabling GameplayDebugger on server, please wait for replicated data..."));
-					const FString ServerCheatString = FString::Printf(TEXT("cheat EnableGDT %s"), *LocalPC->PlayerState->UniqueId.ToString());
-					UE_LOG(LogGameplayDebugger, Warning, TEXT("Sending to Server: %s"), *ServerCheatString);
-					LocalPC->ConsoleCommand(*ServerCheatString);
-					bHandled = true;
-				}
-				else if (!Replicator->IsToolCreated())
+				if (LocalPC && MyWorld)
 				{
-					Replicator->CreateTool();
-					Replicator->EnableTool();
-					bHandled = true;
+					break;
 				}
 			}
 		}
-		else if (Inworld->GetNetMode() < NM_Client)
+	}
+
+	if (MyWorld == nullptr)
+	{
+		return false;
+	}
+
+	if (LocalPC == nullptr)
+	{
+		if (UniquePlayerId.Len() > 0)
 		{
-			APlayerController* LocalPC = NULL;
-			UE_LOG(LogGameplayDebugger, Warning, TEXT("Got from client: %s"), Cmd);
-			FString UniqueId = FParse::Token(Cmd, 0);
-			for (FConstPlayerControllerIterator Iterator = Inworld->GetPlayerControllerIterator(); Iterator; ++Iterator)
+			for (FConstPlayerControllerIterator Iterator = MyWorld->GetPlayerControllerIterator(); Iterator; ++Iterator)
 			{
 				APlayerController* PlayerController = *Iterator;
-				UE_LOG(LogGameplayDebugger, Warning, TEXT("- Client: %s"), *PlayerController->PlayerState->UniqueId.ToString());
-				if (PlayerController && PlayerController->PlayerState->UniqueId.ToString() == UniqueId)
+				UE_LOG(LogGameplayDebugger, Log, TEXT("- Client: %s"), *PlayerController->PlayerState->UniqueId.ToString());
+				if (PlayerController && PlayerController->PlayerState->UniqueId.ToString() == UniquePlayerId)
 				{
 					LocalPC = PlayerController;
 					break;
 				}
 			}
-			if (!LocalPC && Inworld->GetNetMode() != NM_DedicatedServer)
+		}
+
+		if (!LocalPC && MyWorld->GetNetMode() != NM_DedicatedServer)
+		{
+			LocalPC = MyWorld->GetFirstPlayerController();
+		}
+	}
+
+	if (LocalPC == nullptr)
+	{
+		return false;
+	}
+
+	if (MyWorld->GetNetMode() == NM_Client)
+	{
+		AGameplayDebuggingReplicator* Replicator = NULL;
+		for (TActorIterator<AGameplayDebuggingReplicator> It(MyWorld); It; ++It)
+		{
+			Replicator = *It;
+			if (Replicator && !Replicator->IsPendingKill())
 			{
-				LocalPC = Inworld->GetFirstPlayerController();
+				APlayerController* PCOwner = Replicator->GetLocalPlayerOwner();
+				if (LocalPC == PCOwner)
+				{
+					break;
+				}
 			}
+			Replicator = NULL;
+		}
 
-			if (LocalPC)
+		if (!Replicator)
+		{
+			LocalPC->ClientMessage(TEXT("Enabling GameplayDebugger on server, please wait for replicated data..."));
+			const FString ServerCheatString = FString::Printf(TEXT("cheat EnableGDT %s"), *LocalPC->PlayerState->UniqueId.ToString());
+			UE_LOG(LogGameplayDebugger, Warning, TEXT("Sending to Server: %s"), *ServerCheatString);
+			LocalPC->ConsoleCommand(*ServerCheatString);
+			bHandled = true;
+		}
+		else if (!Replicator->IsToolCreated())
+		{
+			Replicator->CreateTool();
+			Replicator->EnableTool();
+			bHandled = true;
+		}
+	}
+	else
+	{
+		UE_LOG(LogGameplayDebugger, Warning, TEXT("Got from client: EnableGDT %s"), *UniquePlayerId);
+		{
+			AGameplayDebuggingReplicator* Replicator = NULL;
+			for (TActorIterator<AGameplayDebuggingReplicator> It(MyWorld); It; ++It)
 			{
-				AGameplayDebuggingReplicator* Replicator = NULL;
-				for (TActorIterator<AGameplayDebuggingReplicator> It(Inworld); It; ++It)
+				Replicator = *It;
+				if (Replicator && !Replicator->IsPendingKill())
+				{
+					APlayerController* PCOwner = Replicator->GetLocalPlayerOwner();
+					if (LocalPC == PCOwner)
+					{
+						break;
+					}
+				}
+				Replicator = NULL;
+			}
+
+			if (!Replicator)
+			{
+				CreateGameplayDebuggerForPlayerController(LocalPC);
+				for (TActorIterator<AGameplayDebuggingReplicator> It(MyWorld); It; ++It)
 				{
 					Replicator = *It;
 					if (Replicator && !Replicator->IsPendingKill())
@@ -529,77 +618,24 @@ bool FGameplayDebugger::Exec(UWorld* Inworld, const TCHAR* Cmd, FOutputDevice& A
 					}
 					Replicator = NULL;
 				}
-
-				if (!Replicator)
-				{
-					CreateGameplayDebuggerForPlayerController(LocalPC);
-					for (TActorIterator<AGameplayDebuggingReplicator> It(Inworld); It; ++It)
-					{
-						Replicator = *It;
-						if (Replicator && !Replicator->IsPendingKill())
-						{
-							APlayerController* PCOwner = Replicator->GetLocalPlayerOwner();
-							if (LocalPC == PCOwner)
-							{
-								break;
-							}
-						}
-						Replicator = NULL;
-					}
-				}
-
-				if (Inworld->GetNetMode() != NM_DedicatedServer)
-				{
-					if (Replicator && !Replicator->IsToolCreated())
-					{
-						Replicator->CreateTool();
-						Replicator->EnableTool();
-						bHandled = true;
-					}
-				}
-				else
-				{
-					if (Replicator)
-					{
-						Replicator->ClientAutoActivate();
-						bHandled = true;
-					}
-				}
 			}
-		}
-	}
-	else if (FParse::Command(&Cmd, TEXT("RunEQS")))
-	{
-		bHandled = true;
-		APlayerController* MyPC = Inworld->GetFirstPlayerController();
-		UAISystem* AISys = UAISystem::GetCurrent(*Inworld);
 
-		UEnvQueryManager* EQS = AISys ? AISys->GetEnvironmentQueryManager() : NULL;
-		if (MyPC && EQS)
-		{
-			AGameplayDebuggingReplicator* DebuggingReplicator = NULL;
-			for (TActorIterator<AGameplayDebuggingReplicator> It(Inworld); It; ++It)
+			if (MyWorld->GetNetMode() != NM_DedicatedServer)
 			{
-				AGameplayDebuggingReplicator* A = *It;
-				if (!A->IsPendingKill())
+				if (Replicator && !Replicator->IsToolCreated())
 				{
-					DebuggingReplicator = A;
-					if (!DebuggingReplicator->IsGlobalInWorld() && DebuggingReplicator->GetLocalPlayerOwner() == MyPC)
-					{
-						break;
-					}
+					Replicator->CreateTool();
+					Replicator->EnableTool();
+					bHandled = true;
 				}
 			}
-
-			UObject* Target = DebuggingReplicator != NULL ? DebuggingReplicator->GetSelectedActorToDebug() : NULL;
-			FString QueryName = FParse::Token(Cmd, 0);
-			if (Target)
-			{
-				AISys->RunEQS(QueryName, Target);
-			}
 			else
 			{
-				MyPC->ClientMessage(TEXT("No debugging target to run EQS"));
+				if (Replicator)
+				{
+					Replicator->ClientAutoActivate();
+					bHandled = true;
+				}
 			}
 		}
 	}
