Small cleanup for Gameplay Debugger, removed leftovers from old changes.

diff --git a/Engine/Source/Developer/GameplayDebugger/Classes/GameplayDebuggingComponent.h b/Engine/Source/Developer/GameplayDebugger/Classes/GameplayDebuggingComponent.h
--- a/Engine/Source/Developer/GameplayDebugger/Classes/GameplayDebuggingComponent.h
+++ b/Engine/Source/Developer/GameplayDebugger/Classes/GameplayDebuggingComponent.h
@@ -124,9 +124,8 @@ class GAMEPLAYDEBUGGER_API UGameplayDebuggingComponent : public UPrimitiveCompon
 	TArray<uint8> EQSRepData;
 	
 	/** local EQS debug data, decoded from EQSRepData blob */
-#if  USE_EQS_DEBUGGER || ENABLE_VISUAL_LOG
 	TArray<EQSDebug::FQueryData> EQSLocalData;	
-#endif
+
 	/** End EQS replication data */
 
 	UPROPERTY(Replicated)
diff --git a/Engine/Source/Developer/GameplayDebugger/Private/GameplayDebuggingComponent.cpp b/Engine/Source/Developer/GameplayDebugger/Private/GameplayDebuggingComponent.cpp
--- a/Engine/Source/Developer/GameplayDebugger/Private/GameplayDebuggingComponent.cpp
+++ b/Engine/Source/Developer/GameplayDebugger/Private/GameplayDebuggingComponent.cpp
@@ -799,7 +799,6 @@ void UGameplayDebuggingComponent::ServerReplicateData(uint32 InMessage, uint32
 void UGameplayDebuggingComponent::OnChangeEQSQuery()
 {
 #if !(UE_BUILD_SHIPPING || UE_BUILD_TEST)
-	AGameplayDebuggingReplicator* Replicator = Cast<AGameplayDebuggingReplicator>(GetOwner());
 	if (++CurrentEQSIndex >= EQSLocalData.Num())
 	{
 		CurrentEQSIndex = 0;
