Fixed monolitics build issue.

diff --git a/Engine/Source/Developer/GameplayDebugger/Classes/GameplayDebuggingComponent.h b/Engine/Source/Developer/GameplayDebugger/Classes/GameplayDebuggingComponent.h
--- a/Engine/Source/Developer/GameplayDebugger/Classes/GameplayDebuggingComponent.h
+++ b/Engine/Source/Developer/GameplayDebugger/Classes/GameplayDebuggingComponent.h
@@ -124,7 +124,9 @@ class GAMEPLAYDEBUGGER_API UGameplayDebuggingComponent : public UPrimitiveCompon
 	TArray<uint8> EQSRepData;
 	
 	/** local EQS debug data, decoded from EQSRepData blob */
-	TArray<EQSDebug::FQueryData> EQSLocalData;	
+#if  USE_EQS_DEBUGGER || ENABLE_VISUAL_LOG
+	TArray<EQSDebug::FQueryData> EQSLocalData;
+#endif
 
 	/** End EQS replication data */
 
