Fixed crash on early shutdown of the editor

diff --git a/Engine/Source/Developer/GameplayDebugger/Private/GameplayDebuggingReplicator.cpp b/Engine/Source/Developer/GameplayDebugger/Private/GameplayDebuggingReplicator.cpp
--- a/Engine/Source/Developer/GameplayDebugger/Private/GameplayDebuggingReplicator.cpp
+++ b/Engine/Source/Developer/GameplayDebugger/Private/GameplayDebuggingReplicator.cpp
@@ -280,7 +280,10 @@ void AGameplayDebuggingReplicator::BeginPlay()
 void AGameplayDebuggingReplicator::BeginDestroy()
 {
 #if !(UE_BUILD_SHIPPING || UE_BUILD_TEST)
-	GEngine->bEnableOnScreenDebugMessages = true;
+	if (GEngine)
+	{
+		GEngine->bEnableOnScreenDebugMessages = true;
+	}
 	if (IsDrawEnabled())
 	{
 		EnableDraw(false);
