Fixed GameplayDebugger-related crashes if a PlayerController happens to be saved along with a level #UE4

UE-17154

diff --git a/Engine/Source/Developer/GameplayDebugger/Private/GameplayDebugger.cpp b/Engine/Source/Developer/GameplayDebugger/Private/GameplayDebugger.cpp
--- a/Engine/Source/Developer/GameplayDebugger/Private/GameplayDebugger.cpp
+++ b/Engine/Source/Developer/GameplayDebugger/Private/GameplayDebugger.cpp
@@ -295,6 +295,7 @@ bool FGameplayDebugger::CreateGameplayDebuggerForPlayerController(APlayerControl
 
 	FActorSpawnParameters SpawnInfo;
 	SpawnInfo.SpawnCollisionHandlingOverride = ESpawnActorCollisionHandlingMethod::AlwaysSpawn;
+	SpawnInfo.ObjectFlags = RF_Transient;
 	SpawnInfo.Name = *FString::Printf(TEXT("GameplayDebuggingReplicator_%s"), *PlayerController->GetName());
 	AGameplayDebuggingReplicator* DestActor = World->SpawnActor<AGameplayDebuggingReplicator>(FVector::ZeroVector, FRotator::ZeroRotator, SpawnInfo);
 	if (DestActor != NULL)
@@ -360,6 +361,7 @@ void FGameplayDebugger::WorldAdded(UWorld* InWorld)
 	// create global replicator on level
 	FActorSpawnParameters SpawnInfo;
 	SpawnInfo.SpawnCollisionHandlingOverride = ESpawnActorCollisionHandlingMethod::AlwaysSpawn;
+	SpawnInfo.ObjectFlags = RF_Transient;
 	SpawnInfo.Name = *FString::Printf(TEXT("GameplayDebuggingReplicator_Global"));
 	AGameplayDebuggingReplicator* DestActor = InWorld->SpawnActor<AGameplayDebuggingReplicator>(FVector::ZeroVector, FRotator::ZeroRotator, SpawnInfo);
 	if (DestActor != NULL)
@@ -397,7 +399,11 @@ void FGameplayDebugger::OnLevelActorAdded(AActor* InActor)
 	APlayerController* PC = Cast<APlayerController>(InActor);
 	if (PC && PC->GetNetMode() < ENetMode::NM_Client)
 	{
-		CreateGameplayDebuggerForPlayerController(PC);
+		UWorld* PCWorld = PC->GetWorld();
+		if (PCWorld && PCWorld->IsGameWorld())
+		{
+			CreateGameplayDebuggerForPlayerController(PC);
+		}
 	}
 #endif
 }
diff --git a/Engine/Source/Developer/GameplayDebugger/Private/GameplayDebuggingReplicator.cpp b/Engine/Source/Developer/GameplayDebugger/Private/GameplayDebuggingReplicator.cpp
--- a/Engine/Source/Developer/GameplayDebugger/Private/GameplayDebuggingReplicator.cpp
+++ b/Engine/Source/Developer/GameplayDebugger/Private/GameplayDebuggingReplicator.cpp
@@ -332,7 +332,7 @@ UGameplayDebuggingComponent* AGameplayDebuggingReplicator::GetDebugComponent()
 #if !(UE_BUILD_SHIPPING || UE_BUILD_TEST)
 	if (IsPendingKill() == false && !DebugComponent && DebugComponentClass.IsValid() && GetNetMode() < ENetMode::NM_Client)
 	{
-		DebugComponent = NewObject<UGameplayDebuggingComponent>(this, DebugComponentClass.Get());
+		DebugComponent = NewObject<UGameplayDebuggingComponent>(this, DebugComponentClass.Get(), TEXT("DebugComponent"), RF_Transient);
 		DebugComponent->SetIsReplicated(true);
 		DebugComponent->RegisterComponent();
 		DebugComponent->Activate();
