Fixed UE-18594 "Gameplay Debugger is hijacking the Canvas"

diff --git a/Engine/Config/BaseEngine.ini b/Engine/Config/BaseEngine.ini
--- a/Engine/Config/BaseEngine.ini
+++ b/Engine/Config/BaseEngine.ini
@@ -2212,7 +2212,7 @@ DebugComponentControllerClassName = "/Script/GameplayDebugger.GameplayDebuggingC
 MenuStartX=10.0
 MenuStartY=10.0
 DebugInfoStartX=20.0
-DebugInfoStartY=20.0
+DebugInfoStartY=60.0
 
 ; These are the defaults for iOS settings, and they need to be in the .ini since UBT reads the .ini settings, without instantiating the class
 ; Note: These should be in IOSEngine, but are placed here because the Settings editor UI cannot write to the Game IOSEngine, meaning settings in Engine/IOSEngine are non-overridable
diff --git a/Engine/Source/Developer/GameplayDebugger/Classes/GameplayDebuggingControllerComponent.h b/Engine/Source/Developer/GameplayDebugger/Classes/GameplayDebuggingControllerComponent.h
--- a/Engine/Source/Developer/GameplayDebugger/Classes/GameplayDebuggingControllerComponent.h
+++ b/Engine/Source/Developer/GameplayDebugger/Classes/GameplayDebuggingControllerComponent.h
@@ -88,6 +88,8 @@ protected:
 	virtual void BindAIDebugViewKeys(class UInputComponent*& InputComponent);
 	AGameplayDebuggingReplicator* GetDebuggingReplicator() const;
 	virtual void ToggleDebugCamera();
+	virtual void ToggleOnScreenDebugMessages();
+	virtual void ToggleGameHUD();
 
 	TWeakObjectPtr<APlayerController> PlayerOwner;
 	TWeakObjectPtr<ADebugCameraController> DebugCameraController;
diff --git a/Engine/Source/Developer/GameplayDebugger/Private/GameplayDebuggingControllerComponent.cpp b/Engine/Source/Developer/GameplayDebugger/Private/GameplayDebuggingControllerComponent.cpp
--- a/Engine/Source/Developer/GameplayDebugger/Private/GameplayDebuggingControllerComponent.cpp
+++ b/Engine/Source/Developer/GameplayDebugger/Private/GameplayDebuggingControllerComponent.cpp
@@ -305,6 +305,8 @@ void UGameplayDebuggingControllerComponent::BindAIDebugViewKeys(class UInputComp
 			InputComponent->BindKey(EKeys::Add, IE_Released, this, &UGameplayDebuggingControllerComponent::CycleDetailsView);
 		}
 		InputComponent->BindKey(FInputChord(EKeys::Tab, false, false, false, false), IE_Released, this, &UGameplayDebuggingControllerComponent::ToggleDebugCamera);
+		InputComponent->BindKey(FInputChord(EKeys::Tab, false, true, false, false), IE_Released, this, &UGameplayDebuggingControllerComponent::ToggleOnScreenDebugMessages);
+		InputComponent->BindKey(FInputChord(EKeys::Tilde, false, true, false, false), IE_Released, this, &UGameplayDebuggingControllerComponent::ToggleGameHUD);
 	}
 #endif //!(UE_BUILD_SHIPPING || UE_BUILD_TEST)
 }
@@ -434,6 +436,32 @@ void UGameplayDebuggingControllerComponent::UpdateNavMeshTimer()
 #endif //!(UE_BUILD_SHIPPING || UE_BUILD_TEST)
 }
 
+void UGameplayDebuggingControllerComponent::ToggleOnScreenDebugMessages()
+{
+	if (!bToolActivated || PlayerOwner == nullptr)
+	{
+		return;
+	}
+
+	if (GEngine)
+	{
+		GEngine->bEnableOnScreenDebugMessages = !GEngine->bEnableOnScreenDebugMessages;
+	}
+}
+
+void UGameplayDebuggingControllerComponent::ToggleGameHUD()
+{
+	if (!bToolActivated || PlayerOwner == nullptr)
+	{
+		return;
+	}
+
+	if (AHUD* const GameHUD = PlayerOwner->GetHUD())
+	{
+		GameHUD->bShowHUD = !GameHUD->bShowHUD;
+	}
+}
+
 void UGameplayDebuggingControllerComponent::ToggleDebugCamera()
 {
 #if !(UE_BUILD_SHIPPING || UE_BUILD_TEST)
diff --git a/Engine/Source/Developer/GameplayDebugger/Private/GameplayDebuggingHUDComponent.cpp b/Engine/Source/Developer/GameplayDebugger/Private/GameplayDebuggingHUDComponent.cpp
--- a/Engine/Source/Developer/GameplayDebugger/Private/GameplayDebuggingHUDComponent.cpp
+++ b/Engine/Source/Developer/GameplayDebugger/Private/GameplayDebuggingHUDComponent.cpp
@@ -97,7 +97,7 @@ void AGameplayDebuggingHUDComponent::PrintAllData()
 #if !(UE_BUILD_SHIPPING || UE_BUILD_TEST)
 
 	// Allow child hud components to position the displayed info
-	DefaultContext = FPrintContext(GEngine->GetSmallFont(), Canvas, DebugInfoStartX, DebugInfoStartY);;
+	DefaultContext = FPrintContext(GEngine->GetSmallFont(), Canvas, DebugInfoStartX, DebugInfoStartY);
 	DefaultContext.FontRenderInfo.bEnableShadow = true;
 
 	if (DefaultContext.Canvas != NULL)
@@ -136,6 +136,7 @@ void AGameplayDebuggingHUDComponent::DrawMenu(const float X, const float Y, clas
 	const float OldX = DefaultContext.CursorX;
 	const float OldY = DefaultContext.CursorY;
 
+	UGameplayDebuggingControllerComponent*  GDC = GetDebuggingReplicator()->FindComponentByClass<UGameplayDebuggingControllerComponent>();
 	if (DefaultContext.Canvas != NULL)
 	{
 		TArray<FDebugCategoryView> Categories;
@@ -152,7 +153,6 @@ void AGameplayDebuggingHUDComponent::DrawMenu(const float X, const float Y, clas
 		FString ActivationKeyName = TEXT("Apostrophe");
 
 		APlayerController* const MyPC = Cast<APlayerController>(PlayerOwner);
-		UGameplayDebuggingControllerComponent*  GDC = GetDebuggingReplicator()->FindComponentByClass<UGameplayDebuggingControllerComponent>();
 		if (GDC)
 		{
 			ActivationKeyDisplayName = GDC->GetActivationKey().Key.GetDisplayName().ToString();
@@ -188,6 +188,26 @@ void AGameplayDebuggingHUDComponent::DrawMenu(const float X, const float Y, clas
 			TotalWidth += CategoriesWidth[DebugCameraIndex];
 			MaxHeight = FMath::Max(MaxHeight, StrHeight);
 		}
+		{
+			const int32 DebugCameraIndex = Categories.Add(FDebugCategoryView());
+			CategoriesWidth.AddZeroed(1);
+			Categories[DebugCameraIndex].Desc = FString::Printf(TEXT(" %s[Ctrl+Tab]: %s  "), GEngine && GEngine->bEnableOnScreenDebugMessages ? TEXT("{Green}") : TEXT("{White}"), TEXT("DebugMessages"));
+			float StrHeight = 0.0f;
+			CalulateStringSize(DefaultContext, DefaultContext.Font, Categories[DebugCameraIndex].Desc, CategoriesWidth[DebugCameraIndex], StrHeight);
+			TotalWidth += CategoriesWidth[DebugCameraIndex];
+			MaxHeight = FMath::Max(MaxHeight, StrHeight);
+		}
+		{
+			const AHUD* GameHUD = MyPC ? MyPC->GetHUD() : NULL;
+			const int32 DebugCameraIndex = Categories.Add(FDebugCategoryView());
+			CategoriesWidth.AddZeroed(1);
+			Categories[DebugCameraIndex].Desc = FString::Printf(TEXT(" %s[Ctrl+Tilde]: %s  "), GameHUD && GameHUD->bShowHUD ? TEXT("{Green}") : TEXT("{White}"), TEXT("GameHUD"));
+			float StrHeight = 0.0f;
+			CalulateStringSize(DefaultContext, DefaultContext.Font, Categories[DebugCameraIndex].Desc, CategoriesWidth[DebugCameraIndex], StrHeight);
+			TotalWidth += CategoriesWidth[DebugCameraIndex];
+			MaxHeight = FMath::Max(MaxHeight, StrHeight);
+		}
+
 
 		TotalWidth = FMath::Max(TotalWidth, HeaderWidth);
 
@@ -214,6 +234,33 @@ void AGameplayDebuggingHUDComponent::DrawMenu(const float X, const float Y, clas
 		PrintString(DefaultContext, "\n{red}No Pawn selected - waiting for data to replicate from server. {green}Press and hold ' to select Pawn \n");
 	}
 
+	if (GDC && GDC->GetDebugCameraController().IsValid())
+	{
+		ADebugCameraController* DebugCamController = GDC->GetDebugCameraController().Get();
+		if (DebugCamController != NULL)
+		{
+			FVector const CamLoc = DebugCamController->PlayerCameraManager->GetCameraLocation();
+			FRotator const CamRot = DebugCamController->PlayerCameraManager->GetCameraRotation();
+
+			FString HitString;
+			FCollisionQueryParams TraceParams(NAME_None, true, this);
+			FHitResult Hit;
+			bool bHit = GetWorld()->LineTraceSingleByChannel(Hit, CamLoc, CamRot.Vector() * 100000.f + CamLoc, ECC_Pawn, TraceParams);
+			if (bHit && Hit.GetActor() != nullptr)
+			{
+				HitString = FString::Printf(TEXT("{white}Under cursor: {yellow}'%s'"), *Hit.GetActor()->GetName());
+				DrawDebugLine(GetWorld(), Hit.Location, Hit.Location + Hit.Normal*30.f, FColor::White);
+			}
+			else
+			{
+				HitString = FString::Printf(TEXT("Not actor under cursor"));
+			}
+
+			PrintString(DefaultContext, FColor::White, HitString, MenuStartX, MenuStartY + 40);
+		}
+	}
+
+
 	DefaultContext.CursorX = OldX;
 	DefaultContext.CursorY = OldY;
 #endif //!(UE_BUILD_SHIPPING || UE_BUILD_TEST)
@@ -227,7 +274,7 @@ void AGameplayDebuggingHUDComponent::DrawDebugComponentData(APlayerController* M
 	const FVector ScreenLoc = SelectedActor ? ProjectLocation(DefaultContext, SelectedActor->GetActorLocation() + FVector(0.f, 0.f, SelectedActor->GetSimpleCollisionHalfHeight())) : FVector::ZeroVector;
 
 	OverHeadContext = FPrintContext(GEngine->GetSmallFont(), Canvas, ScreenLoc.X, ScreenLoc.Y);
-	DefaultContext.CursorY += 20;
+	//DefaultContext.CursorY += 20;
 
 	FGameplayDebuggerSettings DebuggerSettings = GameplayDebuggerSettings(GetDebuggingReplicator());
 	bool bForceOverhead = false;
@@ -423,7 +470,7 @@ void AGameplayDebuggingHUDComponent::DrawBehaviorTreeData(APlayerController* PC,
 	PrintString(DefaultContext, TEXT("\n{green}BEHAVIOR TREE\n"));
 	PrintString(DefaultContext, FString::Printf(TEXT("Brain Component: {yellow}%s\n"), *DebugComponent->BrainComponentName));
 	PrintString(DefaultContext, DebugComponent->BrainComponentString);
-	PrintString(DefaultContext, FColor::White, DebugComponent->BlackboardString, 600.0f, 40.0f);
+	PrintString(DefaultContext, FColor::White, DebugComponent->BlackboardString, 600.0f, DebugInfoStartY);
 #endif //!(UE_BUILD_SHIPPING || UE_BUILD_TEST)
 }
 
