UE4 - fix vsiual log / gameplay debugger #ifdefs

diff --git a/Engine/Source/Developer/GameplayDebugger/Private/GameplayDebuggingComponent.cpp b/Engine/Source/Developer/GameplayDebugger/Private/GameplayDebuggingComponent.cpp
--- a/Engine/Source/Developer/GameplayDebugger/Private/GameplayDebuggingComponent.cpp
+++ b/Engine/Source/Developer/GameplayDebugger/Private/GameplayDebuggingComponent.cpp
@@ -516,7 +516,7 @@ void UGameplayDebuggingComponent::CollectBasicPathData(APawn* MyPawn)
 
 void UGameplayDebuggingComponent::CollectBasicBehaviorData(APawn* MyPawn)
 {
-#if !(UE_BUILD_SHIPPING || UE_BUILD_TEST)
+#if !(UE_BUILD_SHIPPING || UE_BUILD_TEST) && ENABLE_VISUAL_LOG
 	AAIController* MyAIController = Cast<AAIController>(MyPawn->GetController());
 	UBehaviorTreeComponent* BTC = MyAIController ? Cast<UBehaviorTreeComponent>(MyAIController->BrainComponent) : nullptr;
 	bIsUsingBehaviorTree = (BTC != nullptr);
@@ -641,7 +641,7 @@ void UGameplayDebuggingComponent::OnRep_UpdateBlackboard()
 
 void UGameplayDebuggingComponent::CollectPathData()
 {
-#if !(UE_BUILD_SHIPPING || UE_BUILD_TEST)
+#if !(UE_BUILD_SHIPPING || UE_BUILD_TEST) && ENABLE_VISUAL_LOG
 	APawn* MyPawn = Cast<APawn>(GetSelectedActor());
 
 	bool bRefreshRendering = false;
