Ensure that GameplayDebugger is available for UnrealEd
Remove currently inaccurate deprecation warning

diff --git a/Engine/Source/Developer/GameplayDebugger/Private/GameplayDebugger.cpp b/Engine/Source/Developer/GameplayDebugger/Private/GameplayDebugger.cpp
--- a/Engine/Source/Developer/GameplayDebugger/Private/GameplayDebugger.cpp
+++ b/Engine/Source/Developer/GameplayDebugger/Private/GameplayDebugger.cpp
@@ -67,8 +67,8 @@ IMPLEMENT_MODULE(FGameplayDebugger, GameplayDebugger)
 // This code will execute after your module is loaded into memory (but after global variables are initialized, of course.)
 void FGameplayDebugger::StartupModule()
 { 
-	//EMIT_CUSTOM_WARNING("/Engine/Source/Developer/GameplayDebugger module is deprecated and it's going to be removed with next UE4 version. Please use GameplayDebuggerPluging instead.");
-	UE_LOG(LogGameplayDebugger, Warning, TEXT("/Engine/Source/Developer/GameplayDebugger module is deprecated and it's going to be removed with next UE4 version. Please use GameplayDebuggerPluging instead."));
+	//EMIT_CUSTOM_WARNING("/Engine/Source/Developer/GameplayDebugger module is deprecated and it's going to be removed with next UE4 version. Please use GameplayDebuggerPlugin instead.");
+	//UE_LOG(LogGameplayDebugger, Warning, TEXT("/Engine/Source/Developer/GameplayDebugger module is deprecated and it's going to be removed with next UE4 version. Please use GameplayDebuggerPlugin instead."));
 
 #if !(UE_BUILD_SHIPPING || UE_BUILD_TEST)
 	if (GEngine)
diff --git a/Engine/Source/Editor/UnrealEd/UnrealEd.Build.cs b/Engine/Source/Editor/UnrealEd/UnrealEd.Build.cs
--- a/Engine/Source/Editor/UnrealEd/UnrealEd.Build.cs
+++ b/Engine/Source/Editor/UnrealEd/UnrealEd.Build.cs
@@ -83,6 +83,7 @@ public class UnrealEd : ModuleRules
 				"SourceControl",
 				"UnrealEdMessages",
                 "AIModule",
+				"GameplayDebugger",
 				"BlueprintGraph",
                 "Http",
 				"UnrealAudio",
