Copying //UE4/Dev-Framework to //UE4/Main @ 2830052

==========================
MAJOR FEATURES + CHANGES
==========================

Change 2807479 on 2015/12/17 by [REDACTED_PERSON]

	#ue4 - Optimize allocations in FFinalPostProcessSettings. Fix bug skipping an element when updating the ContributingCubemaps array if one is removed during an iteration.

Change 2807695 on 2015/12/17 by [REDACTED_PERSON]

	Changing the VolumeScale sound concurrency implementation to be a volume scalar on older sounds in a concurrency group rather than scale new sounds.

Change 2808628 on 2015/12/18 by [REDACTED_PERSON]

	Minor optimization: Dont call GetWorld() 3 times in a row in USceneComponent::ShouldRender()

Change 2810090 on 2015/12/21 by [REDACTED_PERSON]

	Allow collision and trail particles to run off the game thread.
	Clean up missing marshelled AsyncComponentToWorld

Change 2814518 on 2016/01/04 by [REDACTED_PERSON]

	Use Reset instead of Empty

Change 2814530 on 2016/01/04 by [REDACTED_PERSON]

	UE-2689 Hook up lots of new icons for component classes

Change 2814665 on 2016/01/04 by [REDACTED_PERSON]

	PR #1860:Fix ChildActorComponent linking issue when extended from game (Contributed by [REDACTED_PERSON])

Change 2814686 on 2016/01/04 by [REDACTED_PERSON]

	Fix for radial force components applying impulses multiple times to destructible components

Change 2815221 on 2016/01/04 by [REDACTED_PERSON]

	UE-24528 Fixing focus-distance scaling and focus priority scaling working with concurrency settings.

	USoundAttenuation:
	- Added methods to SoundAttenuation to get focus, attenuation, and distance scale based on focus factor.
	- Added NonFocusDistanceScale parameter to attenuation settings customization so its setable by users

	USoundBase:
	- Added function on USoundBase which checks IsLooping, which checks the duration
	- Removed functions IsAudible and IsAudibleSimple from USoundBase since that functionality is accomplished differently and takes into account focus-distance scaling now.

	FActiveSound:
	- Caching values of FocusPriorityScale and FocusDistanceScale to Active sound so that can be used for audibility checks
	- Included FocusDistanceScale when computing bIsAudible in Active sound
	- Moved GetAttenuationListenerData and the accompanying private struct FAttenuationListenerData to FAudioDevice since the audio device and moved code which computes focus factor and geometry into FAudioDevice::GetFocusFactor since it needs to be computed before an active sound exist

	FAudioDevice:
	- Added SoundIsAudible function which queries a particular sound if its audible based on attenuation settings, focus factor, etc. Will optionally return computed MaxDistance to avoid recomputing it in certain cases.
	- Added a few helper functions: FindClosestListenerIndex, GetAttenuationListenerData
	- Added GetFocusFactor, which performs the vector math to determine focus factor (0.0 is in-focus, 1.0 is out of focus) for a given sound and listener and attenuation focus settings.
	- Updated CreateComponent and PlaySoundAtLocation functions to use the new SoundIsAudible function rather than the old IsAudibleSimple

Change 2815694 on 2016/01/05 by thomas.sarkanen

	Added test texture

Change 2815695 on 2016/01/05 by thomas.sarkanen

	Modified test texture

Change 2815709 on 2016/01/05 by [REDACTED_PERSON]

	PR #1778 : New BP-callablle function AActor::WasRecentlyRendered, with optional tolerance in seconds
	https://github.com/EpicGames/UnrealEngine/pull/1778
	#github 1778

Change 2815711 on 2016/01/05 by [REDACTED_PERSON]

	PR #1534 : Add missing default tolerance value for FVector2D::Equal
	https://github.com/EpicGames/UnrealEngine/pull/1534
	#github 1534

Change 2815714 on 2016/01/05 by [REDACTED_PERSON]

	PR #1887 : Added 'Thickness' Parameter to all valid options in Draw Debug Helper
	https://github.com/EpicGames/UnrealEngine/pull/1887
	#github 1887

Change 2815725 on 2016/01/05 by [REDACTED_PERSON]

	Added comment that USkeletalMeshComponent::bEnablePhysicsOnDedicatedServer cannot be changed at runtime

Change 2815813 on 2016/01/05 by [REDACTED_PERSON]

	Move dispatch of transition functions to new OnMatchStateSet virtual which is called from SetMatchState allowing subclasses to insert their own handling between the set of the variable and dispatch to GameState and Blueprints rather than having to override the entire function.

Change 2815884 on 2016/01/05 by [REDACTED_PERSON]

	- Avoid PSC iterating over all particle emitters each tick to check for changes in DetailMode. Now remember global DetailMode when we last checked, and only iterate again if that has changed.
	- Add 'WarmupTime' and 'CPUCollision' to Asset Registry Tags for Particle Systems

Change 2816306 on 2016/01/05 by [REDACTED_PERSON]

	(4.11) Pass boolean to OnComponentDestroyed that indicates if the entire Actor's hierarchy is being destroyed. If so we can avoid doing a lot of expensive tear down, particularly detaching each component one at a time from the hierarchy and causing massive position/bounds updates.

diff --git a/Engine/Plugins/Runtime/CableComponent/Source/CableComponent/Classes/CableComponent.h b/Engine/Plugins/Runtime/CableComponent/Source/CableComponent/Classes/CableComponent.h
--- a/Engine/Plugins/Runtime/CableComponent/Source/CableComponent/Classes/CableComponent.h
+++ b/Engine/Plugins/Runtime/CableComponent/Source/CableComponent/Classes/CableComponent.h
@@ -97,6 +97,10 @@ public:
 	UPROPERTY(EditAnywhere, BlueprintReadWrite, Category="Cable Rendering", meta=(UIMin = "0.1", UIMax = "8"))
 	float TileMaterial;
 
+	/** Axis (in component space) that is used when building cable mesh geometry. Should be orthogonal to cable direction.  */
+	UPROPERTY(EditAnywhere, AdvancedDisplay, BlueprintReadOnly, Category = "Cable")
+	FVector CableUpDir;
+
 private:
 
 	/** Solve the cable spring constraints */
diff --git a/Engine/Plugins/Runtime/CableComponent/Source/CableComponent/Private/CableComponent.cpp b/Engine/Plugins/Runtime/CableComponent/Source/CableComponent/Private/CableComponent.cpp
--- a/Engine/Plugins/Runtime/CableComponent/Source/CableComponent/Private/CableComponent.cpp
+++ b/Engine/Plugins/Runtime/CableComponent/Source/CableComponent/Private/CableComponent.cpp
@@ -101,6 +101,7 @@ public:
 		, CableWidth(Component->CableWidth)
 		, NumSides(Component->NumSides)
 		, TileMaterial(Component->TileMaterial)
+		, CableUpDir(Component->CableUpDir)
 	{
 		VertexBuffer.NumVerts = GetRequiredVertexCount();
 		IndexBuffer.NumIndices = GetRequiredIndexCount();
@@ -170,8 +171,7 @@ public:
 			const FVector ForwardDir = (InPoints[NextIndex] - InPoints[PrevIndex]).GetSafeNormal();
 
 			// Find a side vector at this point
-			const FVector WorldUp(0,0,1);
-			const FVector RightDir = (WorldUp ^ ForwardDir).GetSafeNormal();
+			const FVector RightDir = (CableUpDir ^ ForwardDir).GetSafeNormal();
 
 			// Find an up vector
 			const FVector UpDir = (ForwardDir ^ RightDir).GetSafeNormal();
@@ -331,6 +331,8 @@ private:
 	int32 NumSides;
 
 	float TileMaterial;
+
+	FVector CableUpDir;
 };
 
 
@@ -353,6 +355,8 @@ UCableComponent::UCableComponent( const FObjectInitializer& ObjectInitializer )
 	SolverIterations = 1;
 	TileMaterial = 1.f;
 
+	CableUpDir = FVector(0,0,1);
+
 	SetCollisionProfileName(UCollisionProfile::NoCollision_ProfileName);
 }
 
diff --git a/Engine/Plugins/Runtime/OculusAudio/Source/OculusAudio/Private/OculusAudio.cpp b/Engine/Plugins/Runtime/OculusAudio/Source/OculusAudio/Private/OculusAudio.cpp
--- a/Engine/Plugins/Runtime/OculusAudio/Source/OculusAudio/Private/OculusAudio.cpp
+++ b/Engine/Plugins/Runtime/OculusAudio/Source/OculusAudio/Private/OculusAudio.cpp
@@ -281,8 +281,11 @@ public:
 
 	void Shutdown() override
 	{
-		// Better have more than 0 instances here
-		check(FOculusAudioPlugin::NumInstances > 0);
+		if (FOculusAudioPlugin::NumInstances == 0)
+		{
+			// This means we failed to load the OVR Audio module during initialization and there's nothing to shutdown.
+			return;
+		}
 
 		FOculusAudioPlugin::NumInstances--;
 
diff --git a/Engine/Source/Developer/AudioFormatOgg/Private/AudioFormatOgg.cpp b/Engine/Source/Developer/AudioFormatOgg/Private/AudioFormatOgg.cpp
--- a/Engine/Source/Developer/AudioFormatOgg/Private/AudioFormatOgg.cpp
+++ b/Engine/Source/Developer/AudioFormatOgg/Private/AudioFormatOgg.cpp
@@ -18,7 +18,7 @@
 static_assert(WITH_OGGVORBIS, "No point in compiling the OGG compressor if we don't have Vorbis.");
 
 // Vorbis encoded sound is about 15% better quality than XMA - adjust the quality setting to get consistent cross platform sound quality
-#define VORBIS_QUALITY_MODIFIER		-15
+#define VORBIS_QUALITY_MODIFIER		0.85
 
 #define SAMPLES_TO_READ		1024
 #define SAMPLE_SIZE			( ( uint32 )sizeof( short ) )
@@ -76,7 +76,7 @@ public:
 			FMemoryWriter CompressedData( CompressedDataStore );
 			uint32 BufferOffset = 0;
 
-			float CompressionQuality = ( float )( QualityInfo.Quality + VORBIS_QUALITY_MODIFIER ) / 100.0f;
+			float CompressionQuality = ( float )( QualityInfo.Quality * VORBIS_QUALITY_MODIFIER ) / 100.0f;
 			CompressionQuality = FMath::Clamp( CompressionQuality, -0.1f, 1.0f );
 
 			vorbis_info_init( &vi );
@@ -246,7 +246,7 @@ public:
 			}
 
 			// Extract the relevant info for compression
-			float CompressionQuality = ( float )( QualityInfo.Quality + VORBIS_QUALITY_MODIFIER ) / 100.0f;
+			float CompressionQuality = ( float )( QualityInfo.Quality * VORBIS_QUALITY_MODIFIER ) / 100.0f;
 			CompressionQuality = FMath::Clamp( CompressionQuality, 0.0f, 1.0f );
 
 			vorbis_info_init( &vi );
diff --git a/Engine/Source/Developer/FunctionalTesting/Private/FuncTestRenderingComponent.cpp b/Engine/Source/Developer/FunctionalTesting/Private/FuncTestRenderingComponent.cpp
--- a/Engine/Source/Developer/FunctionalTesting/Private/FuncTestRenderingComponent.cpp
+++ b/Engine/Source/Developer/FunctionalTesting/Private/FuncTestRenderingComponent.cpp
@@ -94,7 +94,7 @@ FBoxSphereBounds UFuncTestRenderingComponent::CalcBounds(const FTransform& Local
 	const AFunctionalTest* Owner = Cast<AFunctionalTest>(GetOwner());
 	if (Owner)
 	{
-		FBox BoundingBox;
+		FBox BoundingBox(EForceInit::ForceInitToZero);
 
 		TArray<AActor*> RelevantActors;
 		Owner->GatherRelevantActors(RelevantActors);
diff --git a/Engine/Source/Developer/GameplayDebugger/Private/GameplayDebugger.cpp b/Engine/Source/Developer/GameplayDebugger/Private/GameplayDebugger.cpp
--- a/Engine/Source/Developer/GameplayDebugger/Private/GameplayDebugger.cpp
+++ b/Engine/Source/Developer/GameplayDebugger/Private/GameplayDebugger.cpp
@@ -511,144 +511,74 @@ bool FGameplayDebugger::Exec(UWorld* InWorld, const TCHAR* Cmd, FOutputDevice& A
 			}
 		}
 
-		return true;
+		bHandled = true;
 	}
-
-
-	if (FParse::Command(&Cmd, TEXT("EnableGDT")) == false)
+	else if (FParse::Command(&Cmd, TEXT("EnableGDT")))
 	{
-		return false;
-	}
-
-	FString UniquePlayerId = FParse::Token(Cmd, 0);
-	APlayerController* LocalPC = NULL;
-	UWorld* MyWorld = InWorld;
+		FString UniquePlayerId = FParse::Token(Cmd, 0);
+		APlayerController* LocalPC = NULL;
+		UWorld* MyWorld = InWorld;
 
-	if (MyWorld == nullptr)
-	{
-		if (UniquePlayerId.Len() > 0)
+		if (MyWorld == nullptr)
 		{
-			// let's find correct world based on Player Id
-			const TIndirectArray<FWorldContext> WorldContexts = GEngine->GetWorldContexts();
-			for (auto& Context : WorldContexts)
+			if (UniquePlayerId.Len() > 0)
 			{
-				if (Context.WorldType != EWorldType::Game && Context.WorldType != EWorldType::PIE)
+				// let's find correct world based on Player Id
+				const TIndirectArray<FWorldContext> WorldContexts = GEngine->GetWorldContexts();
+				for (const FWorldContext& Context : WorldContexts)
 				{
-					continue;
-				}
+					if (Context.WorldType != EWorldType::Game && Context.WorldType != EWorldType::PIE)
+					{
+						continue;
+					}
 
-				UWorld *CurrentWorld = Context.World();
-				for (FConstPlayerControllerIterator Iterator = CurrentWorld->GetPlayerControllerIterator(); Iterator; ++Iterator)
-				{
-					APlayerController* PC = *Iterator;
-					if (PC && PC->PlayerState->UniqueId.ToString() == UniquePlayerId)
+					UWorld *CurrentWorld = Context.World();
+					for (FConstPlayerControllerIterator Iterator = CurrentWorld->GetPlayerControllerIterator(); Iterator; ++Iterator)
 					{
-						LocalPC = PC;
-						MyWorld = PC->GetWorld();
-						break;
+						APlayerController* PC = *Iterator;
+						if (PC && PC->PlayerState->UniqueId.ToString() == UniquePlayerId)
+						{
+							LocalPC = PC;
+							MyWorld = PC->GetWorld();
+							break;
+						}
 					}
-				}
 
-				if (LocalPC && MyWorld)
-				{
-					break;
+					if (LocalPC && MyWorld)
+					{
+						break;
+					}
 				}
 			}
 		}
-	}
 
-	if (MyWorld == nullptr)
-	{
-		return false;
-	}
-
-	if (LocalPC == nullptr)
-	{
-		if (UniquePlayerId.Len() > 0)
+		if (LocalPC == nullptr && MyWorld != nullptr)
 		{
-			for (FConstPlayerControllerIterator Iterator = MyWorld->GetPlayerControllerIterator(); Iterator; ++Iterator)
+			if (UniquePlayerId.Len() > 0)
 			{
-				APlayerController* PlayerController = *Iterator;
-				UE_LOG(LogGameplayDebugger, Log, TEXT("- Client: %s"), *PlayerController->PlayerState->UniqueId.ToString());
-				if (PlayerController && PlayerController->PlayerState->UniqueId.ToString() == UniquePlayerId)
+				for (FConstPlayerControllerIterator Iterator = MyWorld->GetPlayerControllerIterator(); Iterator; ++Iterator)
 				{
-					LocalPC = PlayerController;
-					break;
+					APlayerController* PlayerController = *Iterator;
+					UE_LOG(LogGameplayDebugger, Log, TEXT("- Client: %s"), *PlayerController->PlayerState->UniqueId.ToString());
+					if (PlayerController && PlayerController->PlayerState->UniqueId.ToString() == UniquePlayerId)
+					{
+						LocalPC = PlayerController;
+						break;
+					}
 				}
 			}
-		}
-
-		if (!LocalPC && MyWorld->GetNetMode() != NM_DedicatedServer)
-		{
-			LocalPC = MyWorld->GetGameInstance() ? MyWorld->GetGameInstance()->GetFirstLocalPlayerController() : nullptr;
-		}
-	}
 
-	if (LocalPC == nullptr)
-	{
-		return false;
-	}
-
-	if (MyWorld->GetNetMode() == NM_Client)
-	{
-		AGameplayDebuggingReplicator* Replicator = NULL;
-		for (TActorIterator<AGameplayDebuggingReplicator> It(MyWorld); It; ++It)
-		{
-			Replicator = *It;
-			if (Replicator && !Replicator->IsPendingKill())
+			if (!LocalPC && MyWorld->GetNetMode() != NM_DedicatedServer)
 			{
-				APlayerController* PCOwner = Replicator->GetLocalPlayerOwner();
-				if (LocalPC == PCOwner)
-				{
-					break;
-				}
+				LocalPC = MyWorld->GetGameInstance() ? MyWorld->GetGameInstance()->GetFirstLocalPlayerController() : nullptr;
 			}
-			Replicator = NULL;
 		}
 
-		if (!Replicator)
-		{
-			LocalPC->ClientMessage(TEXT("Enabling GameplayDebugger on server, please wait for replicated data..."));
-			if (LocalPC->PlayerState)
-			{
-				const FString ServerCheatString = FString::Printf(TEXT("cheat EnableGDT %s"), *LocalPC->PlayerState->UniqueId.ToString());
-				UE_LOG(LogGameplayDebugger, Warning, TEXT("Sending to Server: %s"), *ServerCheatString);
-				LocalPC->ConsoleCommand(*ServerCheatString);
-				bHandled = true;
-			}
-		}
-		else
-		{
-			if (Replicator->IsToolCreated() == false)
-			{
-				Replicator->CreateTool();
-			}
-			Replicator->EnableTool();
-			bHandled = true;
-		}
-	}
-	else
-	{
-		UE_LOG(LogGameplayDebugger, Warning, TEXT("Got from client: EnableGDT %s"), *UniquePlayerId);
+		if (LocalPC != nullptr && MyWorld != nullptr)
 		{
-			AGameplayDebuggingReplicator* Replicator = NULL;
-			for (TActorIterator<AGameplayDebuggingReplicator> It(MyWorld); It; ++It)
+			if (MyWorld->GetNetMode() == NM_Client)
 			{
-				Replicator = *It;
-				if (Replicator && !Replicator->IsPendingKill())
-				{
-					APlayerController* PCOwner = Replicator->GetLocalPlayerOwner();
-					if (LocalPC == PCOwner)
-					{
-						break;
-					}
-				}
-				Replicator = NULL;
-			}
-
-			if (!Replicator)
-			{
-				CreateGameplayDebuggerForPlayerController(LocalPC);
+				AGameplayDebuggingReplicator* Replicator = NULL;
 				for (TActorIterator<AGameplayDebuggingReplicator> It(MyWorld); It; ++It)
 				{
 					Replicator = *It;
@@ -662,26 +592,87 @@ bool FGameplayDebugger::Exec(UWorld* InWorld, const TCHAR* Cmd, FOutputDevice& A
 					}
 					Replicator = NULL;
 				}
-			}
 
-			if (MyWorld->GetNetMode() != NM_DedicatedServer)
-			{
-				if (Replicator && !Replicator->IsToolCreated())
+				if (!Replicator)
 				{
-					Replicator->CreateTool();
-					Replicator->EnableTool();
-					bHandled = true;
+					LocalPC->ClientMessage(TEXT("Enabling GameplayDebugger on server, please wait for replicated data..."));
+					if (LocalPC->PlayerState)
+					{
+						const FString ServerCheatString = FString::Printf(TEXT("cheat EnableGDT %s"), *LocalPC->PlayerState->UniqueId.ToString());
+						UE_LOG(LogGameplayDebugger, Warning, TEXT("Sending to Server: %s"), *ServerCheatString);
+						LocalPC->ConsoleCommand(*ServerCheatString);
+					}
+				}
+				else
+				{
+					if (Replicator->IsToolCreated() == false)
+					{
+						Replicator->CreateTool();
+						Replicator->EnableTool();
+					}
+					else
+					{
+						Replicator->EnableDraw(!Replicator->IsDrawEnabled());
+					}
 				}
 			}
 			else
 			{
-				if (Replicator)
+				UE_LOG(LogGameplayDebugger, Warning, TEXT("Got from client: EnableGDT %s"), *UniquePlayerId);
 				{
-					Replicator->ClientAutoActivate();
-					bHandled = true;
+					AGameplayDebuggingReplicator* Replicator = NULL;
+					for (TActorIterator<AGameplayDebuggingReplicator> It(MyWorld); It; ++It)
+					{
+						Replicator = *It;
+						if (Replicator && !Replicator->IsPendingKill())
+						{
+							APlayerController* PCOwner = Replicator->GetLocalPlayerOwner();
+							if (LocalPC == PCOwner)
+							{
+								break;
+							}
+						}
+						Replicator = NULL;
+					}
+
+					if (!Replicator)
+					{
+						CreateGameplayDebuggerForPlayerController(LocalPC);
+						for (TActorIterator<AGameplayDebuggingReplicator> It(MyWorld); It; ++It)
+						{
+							Replicator = *It;
+							if (Replicator && !Replicator->IsPendingKill())
+							{
+								APlayerController* PCOwner = Replicator->GetLocalPlayerOwner();
+								if (LocalPC == PCOwner)
+								{
+									break;
+								}
+							}
+							Replicator = NULL;
+						}
+					}
+
+					if (MyWorld->GetNetMode() != NM_DedicatedServer)
+					{
+						if (Replicator && !Replicator->IsToolCreated())
+						{
+							Replicator->CreateTool();
+							Replicator->EnableTool();
+						}
+					}
+					else
+					{
+						if (Replicator)
+						{
+							Replicator->ClientAutoActivate();
+						}
+					}
 				}
 			}
 		}
+
+		bHandled = true;
 	}
 
 	return bHandled;
diff --git a/Engine/Source/Developer/GameplayDebugger/Private/GameplayDebuggingReplicator.cpp b/Engine/Source/Developer/GameplayDebugger/Private/GameplayDebuggingReplicator.cpp
--- a/Engine/Source/Developer/GameplayDebugger/Private/GameplayDebuggingReplicator.cpp
+++ b/Engine/Source/Developer/GameplayDebugger/Private/GameplayDebuggingReplicator.cpp
@@ -317,8 +317,15 @@ void AGameplayDebuggingReplicator::ClientAutoActivate_Implementation()
 	// we are already replicated so let's activate tool
 	if (GetWorld() && GetNetMode() == ENetMode::NM_Client && !IsToolCreated() && !IsGlobalInWorld())
 	{
-		CreateTool();
-		EnableTool();
+		if (IsToolCreated())
+		{
+			EnableDraw(!IsDrawEnabled());
+		}
+		else
+		{
+			CreateTool();
+			EnableTool();
+		}
 	}
 #endif
 }
@@ -329,8 +336,15 @@ void AGameplayDebuggingReplicator::OnRep_AutoActivate()
 	// we are already replicated so let's activate tool
 	if (GetWorld() && GetNetMode() == ENetMode::NM_Client && !IsToolCreated() && !IsGlobalInWorld())
 	{
-		CreateTool();
-		EnableTool();
+		if (IsToolCreated())
+		{
+			EnableDraw(!IsDrawEnabled());
+		}
+		else
+		{
+			CreateTool();
+			EnableTool();
+		}
 	}
 #endif
 }
diff --git a/Engine/Source/Developer/LogVisualizer/Private/SVisualLogger.cpp b/Engine/Source/Developer/LogVisualizer/Private/SVisualLogger.cpp
--- a/Engine/Source/Developer/LogVisualizer/Private/SVisualLogger.cpp
+++ b/Engine/Source/Developer/LogVisualizer/Private/SVisualLogger.cpp
@@ -653,10 +653,10 @@ void SVisualLogger::HandleSaveCommand(bool bSaveAllData)
 
 		if (bSaved)
 		{
-			if (SaveFilenames.Num() > 0)
+			if (SaveFilenames.Num() > 0 && SaveFilenames[0].IsEmpty() == false)
 			{
 				TArray<FVisualLogDevice::FVisualLogEntryItem> FrameCache;
-				for (auto CurrentName : SelectedRows)
+				for (const FName& CurrentName : SelectedRows)
 				{
 					FVisualLoggerDBRow& DataRow = FVisualLoggerDatabase::Get().GetRowByName(CurrentName);
 					FrameCache.Append(DataRow.GetItems());
@@ -665,10 +665,17 @@ void SVisualLogger::HandleSaveCommand(bool bSaveAllData)
 				if (FrameCache.Num())
 				{
 					FArchive* FileArchive = IFileManager::Get().CreateFileWriter(*SaveFilenames[0]);
-					FVisualLoggerHelpers::Serialize(*FileArchive, FrameCache);
-					FileArchive->Close();
-					delete FileArchive;
-					FileArchive = NULL;
+					if (ensure(FileArchive))
+					{
+						FVisualLoggerHelpers::Serialize(*FileArchive, FrameCache);
+						FileArchive->Close();
+						delete FileArchive;
+						FileArchive = NULL;
+					}
+					else
+					{
+						UE_LOG(LogVisualLogger, Error, TEXT("Failed to create file \"%s\""), *SaveFilenames[0]);
+					}
 				}
 			}
 		}
diff --git a/Engine/Source/Editor/AnimGraph/Classes/AnimGraphNode_AnimDynamics.h b/Engine/Source/Editor/AnimGraph/Classes/AnimGraphNode_AnimDynamics.h
--- a/Engine/Source/Editor/AnimGraph/Classes/AnimGraphNode_AnimDynamics.h
+++ b/Engine/Source/Editor/AnimGraph/Classes/AnimGraphNode_AnimDynamics.h
@@ -88,4 +88,7 @@ protected:
 	// UAnimGraphNode_SkeletalControlBase protected interface
 	virtual const FAnimNode_SkeletalControlBase* GetNode() const override { return &Node; }
 	// End of UAnimGraphNode_SkeletalControlBase protected interface
+
+private:
+	FNodeTitleTextTable CachedNodeTitles;
 };
\ No newline at end of file
diff --git a/Engine/Source/Editor/AnimGraph/Classes/AnimGraphNode_SkeletalControlBase.h b/Engine/Source/Editor/AnimGraph/Classes/AnimGraphNode_SkeletalControlBase.h
--- a/Engine/Source/Editor/AnimGraph/Classes/AnimGraphNode_SkeletalControlBase.h
+++ b/Engine/Source/Editor/AnimGraph/Classes/AnimGraphNode_SkeletalControlBase.h
@@ -90,11 +90,11 @@ protected:
 	// local conversion function for drawing
 	void ConvertToComponentSpaceTransform(const USkeletalMeshComponent* SkelComp, const FTransform & InTransform, FTransform & OutCSTransform, int32 BoneIndex, EBoneControlSpace Space) const;
 	// convert drag vector in component space to bone space 
-	FVector ConvertCSVectorToBoneSpace(const USkeletalMeshComponent* SkelComp, FVector& InCSVector, FCSPose<FCompactPose>& MeshBases, const FName& BoneName, const EBoneControlSpace Space);
+	FVector ConvertCSVectorToBoneSpace(const USkeletalMeshComponent* SkelComp, FVector& InCSVector, FCSPose<FCompactHeapPose>& MeshBases, const FName& BoneName, const EBoneControlSpace Space);
 	// convert rotator in component space to bone space 
-	FQuat ConvertCSRotationToBoneSpace(const USkeletalMeshComponent* SkelComp, FRotator& InCSRotator, FCSPose<FCompactPose>& MeshBases, const FName& BoneName, const EBoneControlSpace Space);
+	FQuat ConvertCSRotationToBoneSpace(const USkeletalMeshComponent* SkelComp, FRotator& InCSRotator, FCSPose<FCompactHeapPose>& MeshBases, const FName& BoneName, const EBoneControlSpace Space);
 	// convert widget location according to bone control space
-	FVector ConvertWidgetLocation(const USkeletalMeshComponent* InSkelComp, FCSPose<FCompactPose>& InMeshBases, const FName& BoneName, const FVector& InLocation, const EBoneControlSpace Space);
+	FVector ConvertWidgetLocation(const USkeletalMeshComponent* InSkelComp, FCSPose<FCompactHeapPose>& InMeshBases, const FName& BoneName, const FVector& InLocation, const EBoneControlSpace Space);
 	// set literal value for FVector
 	void SetDefaultValue(const FString& InDefaultValueName, const FVector& InValue);
 	// get literal value for vector
diff --git a/Engine/Source/Editor/AnimGraph/Classes/AnimStateNode.h b/Engine/Source/Editor/AnimGraph/Classes/AnimStateNode.h
--- a/Engine/Source/Editor/AnimGraph/Classes/AnimStateNode.h
+++ b/Engine/Source/Editor/AnimGraph/Classes/AnimStateNode.h
@@ -35,6 +35,10 @@ public:
 	UPROPERTY(EditAnywhere, Category=Events)
 	FAnimNotifyEvent StateFullyBlended;
 
+	// Whether or not this state will ALWAYS reset it's state on reentry, regardless of remaining weight
+	UPROPERTY(EditAnywhere, AdvancedDisplay, Category = Events)
+	bool bAlwaysResetOnEntry;
+
 	//~ Begin UObject Interface
 	virtual void PostEditChangeProperty(struct FPropertyChangedEvent& PropertyChangedEvent) override;
 	//~ End UObject Interface
diff --git a/Engine/Source/Editor/AnimGraph/Private/AnimGraphNode_AnimDynamics.cpp b/Engine/Source/Editor/AnimGraph/Private/AnimGraphNode_AnimDynamics.cpp
--- a/Engine/Source/Editor/AnimGraph/Private/AnimGraphNode_AnimDynamics.cpp
+++ b/Engine/Source/Editor/AnimGraph/Private/AnimGraphNode_AnimDynamics.cpp
@@ -237,11 +237,49 @@ void UAnimGraphNode_AnimDynamics::ValidateAnimNodeDuringCompilation(USkeleton* F
 
 FText UAnimGraphNode_AnimDynamics::GetNodeTitle(ENodeTitleType::Type TitleType) const
 {
-	return GetControllerDescription();
+	FFormatNamedArguments Arguments;
+	Arguments.Add(TEXT("ControllerDescription"), GetControllerDescription());
+	Arguments.Add(TEXT("BoundBoneName"), FText::FromName(Node.BoundBone.BoneName));
+	if(Node.bChain)
+	{
+		Arguments.Add(TEXT("ChainEndBoneName"), FText::FromName(Node.ChainEnd.BoneName));
+	}
+
+	if(TitleType == ENodeTitleType::ListView || TitleType == ENodeTitleType::MenuTitle)
+	{
+		if(Node.BoundBone.BoneName == NAME_None || (Node.bChain && Node.ChainEnd.BoneName == NAME_None))
+		{
+			return GetControllerDescription();
+		}
+
+		if(Node.bChain)
+		{
+			CachedNodeTitles.SetCachedTitle(TitleType, FText::Format(LOCTEXT("AnimDynamicsNodeTitleSmallChain", "{ControllerDescription} - Chain: {BoundBoneName} -> {ChainEndBoneName}"), Arguments), this);
+		}
+		else
+		{
+			CachedNodeTitles.SetCachedTitle(TitleType, FText::Format(LOCTEXT("AnimDynamicsNodeTitleSmall", "{ControllerDescription} - Bone: {BoundBoneName}"), Arguments), this);
+		}
+	}
+	else
+	{
+		if(Node.bChain)
+		{
+			CachedNodeTitles.SetCachedTitle(TitleType, FText::Format(LOCTEXT("AnimDynamicsNodeTitleLargeChain", "{ControllerDescription}\nChain: {BoundBoneName} -> {ChainEndBoneName}"), Arguments), this);
+		}
+		else
+		{
+			CachedNodeTitles.SetCachedTitle(TitleType, FText::Format(LOCTEXT("AnimDynamicsNodeTitleLarge", "{ControllerDescription}\nBone: {BoundBoneName}"), Arguments), this);
+		}
+	}
+
+	return CachedNodeTitles[TitleType];
 }
 
 void UAnimGraphNode_AnimDynamics::PostEditChangeProperty(struct FPropertyChangedEvent& PropertyChangedEvent)
 {
+	Super::PostEditChangeProperty(PropertyChangedEvent);
+
 	// Regenerate render shape(s)
 	EditPreviewShape = FAnimPhysShape::MakeBox(Node.BoxExtents);
 }
diff --git a/Engine/Source/Editor/AnimGraph/Private/AnimGraphNode_ModifyBone.cpp b/Engine/Source/Editor/AnimGraph/Private/AnimGraphNode_ModifyBone.cpp
--- a/Engine/Source/Editor/AnimGraph/Private/AnimGraphNode_ModifyBone.cpp
+++ b/Engine/Source/Editor/AnimGraph/Private/AnimGraphNode_ModifyBone.cpp
@@ -105,7 +105,7 @@ FVector UAnimGraphNode_ModifyBone::GetWidgetLocation(const USkeletalMeshComponen
 	// if the current widget mode is translate, then shows the widget according to translation space
 	if (CurWidgetMode == FWidget::WM_Translate)
 	{
-		FCSPose<FCompactPose>& MeshBases = AnimNode->ForwardedPose;
+		FCSPose<FCompactHeapPose>& MeshBases = AnimNode->ForwardedPose;
 		WidgetLoc = ConvertWidgetLocation(SkelComp, MeshBases, Node.BoneToModify.BoneName, GetNodeValue(FString("Translation"), Node.Translation), Node.TranslationSpace);
 
 		if (MeshBases.GetPose().IsValid() && Node.TranslationMode == BMM_Additive)
diff --git a/Engine/Source/Editor/AnimGraph/Private/AnimGraphNode_SkeletalControlBase.cpp b/Engine/Source/Editor/AnimGraph/Private/AnimGraphNode_SkeletalControlBase.cpp
--- a/Engine/Source/Editor/AnimGraph/Private/AnimGraphNode_SkeletalControlBase.cpp
+++ b/Engine/Source/Editor/AnimGraph/Private/AnimGraphNode_SkeletalControlBase.cpp
@@ -140,7 +140,7 @@ void UAnimGraphNode_SkeletalControlBase::ConvertToComponentSpaceTransform(const
 }
 
 
-FVector UAnimGraphNode_SkeletalControlBase::ConvertCSVectorToBoneSpace(const USkeletalMeshComponent* SkelComp, FVector& InCSVector, FCSPose<FCompactPose>& MeshBases, const FName& BoneName, const EBoneControlSpace Space)
+FVector UAnimGraphNode_SkeletalControlBase::ConvertCSVectorToBoneSpace(const USkeletalMeshComponent* SkelComp, FVector& InCSVector, FCSPose<FCompactHeapPose>& MeshBases, const FName& BoneName, const EBoneControlSpace Space)
 {
 	FVector OutVector = InCSVector;
 
@@ -180,7 +180,7 @@ FVector UAnimGraphNode_SkeletalControlBase::ConvertCSVectorToBoneSpace(const USk
 	return OutVector;
 }
 
-FQuat UAnimGraphNode_SkeletalControlBase::ConvertCSRotationToBoneSpace(const USkeletalMeshComponent* SkelComp, FRotator& InCSRotator, FCSPose<FCompactPose>& MeshBases, const FName& BoneName, const EBoneControlSpace Space)
+FQuat UAnimGraphNode_SkeletalControlBase::ConvertCSRotationToBoneSpace(const USkeletalMeshComponent* SkelComp, FRotator& InCSRotator, FCSPose<FCompactHeapPose>& MeshBases, const FName& BoneName, const EBoneControlSpace Space)
 {
 	FQuat OutQuat = FQuat::Identity;
 
@@ -235,7 +235,7 @@ FQuat UAnimGraphNode_SkeletalControlBase::ConvertCSRotationToBoneSpace(const USk
 	return OutQuat;
 }
 
-FVector UAnimGraphNode_SkeletalControlBase::ConvertWidgetLocation(const USkeletalMeshComponent* SkelComp, FCSPose<FCompactPose>& MeshBases, const FName& BoneName, const FVector& Location, const EBoneControlSpace Space)
+FVector UAnimGraphNode_SkeletalControlBase::ConvertWidgetLocation(const USkeletalMeshComponent* SkelComp, FCSPose<FCompactHeapPose>& MeshBases, const FName& BoneName, const FVector& Location, const EBoneControlSpace Space)
 {
 	FVector WidgetLoc = FVector::ZeroVector;
 
diff --git a/Engine/Source/Editor/DetailCustomizations/Private/AttenuationSettingsCustomizations.cpp b/Engine/Source/Editor/DetailCustomizations/Private/AttenuationSettingsCustomizations.cpp
--- a/Engine/Source/Editor/DetailCustomizations/Private/AttenuationSettingsCustomizations.cpp
+++ b/Engine/Source/Editor/DetailCustomizations/Private/AttenuationSettingsCustomizations.cpp
@@ -210,12 +210,18 @@ void FAttenuationSettingsCustomization::CustomizeChildren( TSharedRef<IPropertyH
 	ChildBuilder.AddChildProperty(PropertyHandles.FindChecked(GET_MEMBER_NAME_CHECKED(FAttenuationSettings, FocusDistanceScale)).ToSharedRef())
 		.EditCondition(GetIsFocusEnabledAttribute(), nullptr);
 
+	ChildBuilder.AddChildProperty(PropertyHandles.FindChecked(GET_MEMBER_NAME_CHECKED(FAttenuationSettings, NonFocusDistanceScale)).ToSharedRef())
+		.EditCondition(GetIsFocusEnabledAttribute(), nullptr);
+
 	ChildBuilder.AddChildProperty(PropertyHandles.FindChecked(GET_MEMBER_NAME_CHECKED(FAttenuationSettings, FocusPriorityScale)).ToSharedRef())
 		.EditCondition(GetIsFocusEnabledAttribute(), nullptr);
 
 	ChildBuilder.AddChildProperty(PropertyHandles.FindChecked(GET_MEMBER_NAME_CHECKED(FAttenuationSettings, NonFocusPriorityScale)).ToSharedRef())
 		.EditCondition(GetIsFocusEnabledAttribute(), nullptr);
 
+	ChildBuilder.AddChildProperty(PropertyHandles.FindChecked(GET_MEMBER_NAME_CHECKED(FAttenuationSettings, FocusVolumeAttenuation)).ToSharedRef())
+		.EditCondition(GetIsFocusEnabledAttribute(), nullptr);
+
 	ChildBuilder.AddChildProperty(PropertyHandles.FindChecked(GET_MEMBER_NAME_CHECKED(FAttenuationSettings, NonFocusVolumeAttenuation)).ToSharedRef())
 		.EditCondition(GetIsFocusEnabledAttribute(), nullptr);
 
@@ -233,7 +239,7 @@ void FAttenuationSettingsCustomization::CustomizeChildren( TSharedRef<IPropertyH
 	ChildBuilder.AddChildProperty(PropertyHandles.FindChecked(GET_MEMBER_NAME_CHECKED(FAttenuationSettings, bUseComplexCollisionForOcclusion)).ToSharedRef())
 		.EditCondition(GetIsOcclusionEnabledAttribute(), nullptr);
 
-	if (PropertyHandles.Num() != 30)
+	if (PropertyHandles.Num() != 31)
 	{
 		FString PropertyList;
 		for (auto It(PropertyHandles.CreateConstIterator()); It; ++It)
diff --git a/Engine/Source/Editor/EditorStyle/Private/SlateEditorStyle.cpp b/Engine/Source/Editor/EditorStyle/Private/SlateEditorStyle.cpp
--- a/Engine/Source/Editor/EditorStyle/Private/SlateEditorStyle.cpp
+++ b/Engine/Source/Editor/EditorStyle/Private/SlateEditorStyle.cpp
@@ -5348,17 +5348,19 @@ void FSlateEditorStyle::FStyle::SetupClassIconsAndThumbnails()
 		// Asset Type Classes
 		const TCHAR* AssetTypes[] = {
 
+			TEXT("AbilitySystemComponent"),
 			TEXT("Actor"),
 			TEXT("ActorComponent"),
-			TEXT("LevelSequence"),
 			TEXT("AIController"),
 			TEXT("AimOffsetBlendSpace"),
 			TEXT("AimOffsetBlendSpace1D"),
+			TEXT("AIPerceptionComponent"),
 			TEXT("AmbientSound"),
 			TEXT("AnimBlueprint"),
 			TEXT("AnimComposite"),
 			TEXT("AnimMontage"),
 			TEXT("AnimSequence"),
+			TEXT("ApplicationLifecycleComponent"),
 			TEXT("AtmosphericFog"),
 			TEXT("AudioVolume"),
 			TEXT("BehaviorTree"),
@@ -5378,6 +5380,7 @@ void FSlateEditorStyle::FStyle::SetupClassIconsAndThumbnails()
 			TEXT("CameraAnim"),
 			TEXT("CameraBlockingVolume"),
 			TEXT("Character"),
+			TEXT("CharacterMovementComponent"),
 			TEXT("Class"),
 			TEXT("CullDistanceVolume"),
 			TEXT("CurveBase"),
@@ -5406,6 +5409,7 @@ void FSlateEditorStyle::FStyle::SetupClassIconsAndThumbnails()
 			TEXT("KillZVolume"),
 			TEXT("Landscape"),
 			TEXT("LevelBounds"),
+			TEXT("LevelSequence"),
 			TEXT("LevelStreamingVolume"),
 			TEXT("LightmassCharacterIndirectDetailVolume"),
 			TEXT("LightmassImportanceVolume"),
@@ -5423,16 +5427,20 @@ void FSlateEditorStyle::FStyle::SetupClassIconsAndThumbnails()
 			TEXT("MultiFont"),
 			TEXT("NavLinkProxy"),
 			TEXT("NavMeshBoundsVolume"),
+			TEXT("NavModifierComponent"),
 			TEXT("NavModifierVolume"),
 			TEXT("Note"),
 			TEXT("ObjectLibrary"),
 			TEXT("PainCausingVolume"),
 			TEXT("ParticleSystem"),
 			TEXT("Pawn"),
+			TEXT("PawnNoiseEmitterComponent"),
+			TEXT("PawnSensingComponent"),
 			TEXT("PhysicalMaterial"),
 			TEXT("PhysicsAsset"),
 			TEXT("PhysicsConstraintActor"),
 			TEXT("PhysicsConstraintComponent"),
+			TEXT("PhysicsHandleComponent"),
 			TEXT("PhysicsThruster"),
 			TEXT("PhysicsThrusterComponent"),
 			TEXT("PhysicsVolume"),
@@ -5445,10 +5453,12 @@ void FSlateEditorStyle::FStyle::SetupClassIconsAndThumbnails()
 			TEXT("PrecomputedVisibilityVolume"),
 			TEXT("ProceduralFoliageVolume"),
 			TEXT("ProceduralFoliageBlockingVolume"),
+			TEXT("ProjectileMovementComponent"),
 			TEXT("RadialForceActor"),
 			TEXT("RadialForceComponent"),
 			TEXT("ReflectionCapture"),
 			TEXT("ReverbEffect"),
+			TEXT("RotatingMovementComponent"),
 			TEXT("SceneCapture2D"),
 			TEXT("SceneCaptureCube"),
 			TEXT("SceneComponent"),
@@ -5468,6 +5478,7 @@ void FSlateEditorStyle::FStyle::SetupClassIconsAndThumbnails()
 			TEXT("SpotLightMovable"),
 			TEXT("SpotLightStatic"),
 			TEXT("SpotLightStationary"),
+			TEXT("SpringArmComponent"),
 			TEXT("StaticMesh"),
 			TEXT("SubsurfaceProfile"),
 			TEXT("TargetPoint"),
diff --git a/Engine/Source/Editor/IntroTutorials/Private/STutorialRoot.cpp b/Engine/Source/Editor/IntroTutorials/Private/STutorialRoot.cpp
--- a/Engine/Source/Editor/IntroTutorials/Private/STutorialRoot.cpp
+++ b/Engine/Source/Editor/IntroTutorials/Private/STutorialRoot.cpp
@@ -375,11 +375,11 @@ void STutorialRoot::GoToNextStage(TWeakPtr<SWindow> InNavigationWindow)
 			GetMutableDefault<UTutorialStateSettings>()->RecordProgress(CurrentTutorial, CurrentTutorialStage);
 			break;
 		}
-		if (CurrentTutorialStage + 1 >= CurrentTutorial->Stages.Num())
+		if (CurrentTutorialStage >= CurrentTutorial->Stages.Num())
 		{
 			// We went out of bounds for this tutorial, so see if we want to go to another one.
 			CurrentTutorialStage = CurrentTutorial->Stages.Num() - 1;
-			if (FName(*CurrentTutorial->PreviousTutorial.ToString()) != NAME_None)
+			if (FName(*CurrentTutorial->NextTutorial.ToString()) != NAME_None)
 			{
 				TSubclassOf<UEditorTutorial> NextTutorialClass = LoadClass<UEditorTutorial>(NULL, *CurrentTutorial->NextTutorial.ToString(), NULL, LOAD_None, NULL);
 				if (NextTutorialClass != nullptr)
diff --git a/Engine/Source/Editor/UnrealEd/Private/Animation/DebugSkelMeshComponent.cpp b/Engine/Source/Editor/UnrealEd/Private/Animation/DebugSkelMeshComponent.cpp
--- a/Engine/Source/Editor/UnrealEd/Private/Animation/DebugSkelMeshComponent.cpp
+++ b/Engine/Source/Editor/UnrealEd/Private/Animation/DebugSkelMeshComponent.cpp
@@ -421,7 +421,7 @@ void UDebugSkelMeshComponent::GenSpaceBases(TArray<FTransform>& OutSpaceBases)
 	TempLocalAtoms.AddUninitialized(OutSpaceBases.Num());
 	TArray<FActiveVertexAnim> TempVertexAnims;
 	FVector TempRootBoneTranslation;
-	FBlendedCurve TempCurve;
+	FBlendedHeapCurve TempCurve;
 	PerformAnimationEvaluation(SkeletalMesh, AnimScriptInstance, OutSpaceBases, TempLocalAtoms, TempVertexAnims, TempRootBoneTranslation, TempCurve);
 }
 
diff --git a/Engine/Source/Editor/UnrealEd/Private/ConvexDecompTool.cpp b/Engine/Source/Editor/UnrealEd/Private/ConvexDecompTool.cpp
--- a/Engine/Source/Editor/UnrealEd/Private/ConvexDecompTool.cpp
+++ b/Engine/Source/Editor/UnrealEd/Private/ConvexDecompTool.cpp
@@ -83,13 +83,27 @@ void DecomposeMeshToHulls(UBodySetup* InBodySetup, const TArray<FVector>& InVert
 
 	bool bSuccess = false;
 
+	// Validate input by checking bounding box
+	FBox VertBox(0);
+	for(FVector Vert : InVertices)
+	{
+		VertBox += Vert;
+	}
+
+	// If the largest dimension is less than 1 unit, skip trying to generate collision
+	if(VertBox.IsValid != 0 && VertBox.GetSize().GetMax() < 1.f)
+	{
+		return;
+	}
+
+
 #if USE_VHACD
 	FVHACDProgressCallback VHACD_Callback;
 
 	IVHACD::Parameters VHACD_Params;
-	VHACD_Params.m_resolution = 1000000;
-	VHACD_Params.m_maxNumVerticesPerCH = InMaxHullVerts;
-	VHACD_Params.m_concavity = 0.3f * (1.f - FMath::Clamp(InAccuracy, 0.f, 1.f));
+	VHACD_Params.m_resolution = 1000000; // Maximum number of voxels generated during the voxelization stage (default=100,000, range=10,000-16,000,000)
+	VHACD_Params.m_maxNumVerticesPerCH = InMaxHullVerts; // Controls the maximum number of triangles per convex-hull (default=64, range=4-1024)
+	VHACD_Params.m_concavity = 0.3f * (1.f - FMath::Clamp(InAccuracy, 0.f, 1.f)); // Maximum allowed concavity (default=0.0025, range=0.0-1.0)
 	VHACD_Params.m_callback = &VHACD_Callback;
 	VHACD_Params.m_oclAcceleration = false;
 
diff --git a/Engine/Source/Runtime/AIModule/Private/Navigation/CrowdFollowingComponent.cpp b/Engine/Source/Runtime/AIModule/Private/Navigation/CrowdFollowingComponent.cpp
--- a/Engine/Source/Runtime/AIModule/Private/Navigation/CrowdFollowingComponent.cpp
+++ b/Engine/Source/Runtime/AIModule/Private/Navigation/CrowdFollowingComponent.cpp
@@ -639,51 +639,54 @@ void LogPathPartHelper(AActor* LogOwner, FNavMeshPath* NavMeshPath, int32 StartI
 	const FVector CorridorOffset = NavigationDebugDrawing::PathOffset * 1.25f;
 	int32 NumAreaMark = 1;
 
-	NavMesh->BeginBatchQuery();
 	FVisualLogEntry* Snapshot = VisualLogger.GetEntryToWrite(LogOwner, LogOwner->GetWorld()->GetTimeSeconds());
-
-	TArray<FVector> Verts;
-	for (int32 Idx = StartIdx; Idx <= EndIdx; Idx++)
+	if (Snapshot)
 	{
-		const uint8 AreaID = NavMesh->GetPolyAreaID(NavMeshPath->PathCorridor[Idx]);
-		const UClass* AreaClass = NavMesh->GetAreaClass(AreaID);
-
-		Verts.Reset();
-		NavMesh->GetPolyVerts(NavMeshPath->PathCorridor[Idx], Verts);
+		NavMesh->BeginBatchQuery();
 
-		FVector CenterPt = FVector::ZeroVector;
-		for (int32 VIdx = 0; VIdx < Verts.Num(); VIdx++)
+		TArray<FVector> Verts;
+		for (int32 Idx = StartIdx; Idx <= EndIdx; Idx++)
 		{
-			Verts[VIdx].Z += 5.0f;
-			CenterPt += Verts[VIdx];
-		}
-		CenterPt /= Verts.Num();
+			const uint8 AreaID = NavMesh->GetPolyAreaID(NavMeshPath->PathCorridor[Idx]);
+			const UClass* AreaClass = NavMesh->GetAreaClass(AreaID);
 
-		const UNavArea* DefArea = AreaClass ? ((UClass*)AreaClass)->GetDefaultObject<UNavArea>() : NULL;
-		const FColor PolygonColor = AreaClass != UNavigationSystem::GetDefaultWalkableArea() ? (DefArea ? DefArea->DrawColor : NavMesh->GetConfig().Color) : FColorList::LightSteelBlue;
+			Verts.Reset();
+			NavMesh->GetPolyVerts(NavMeshPath->PathCorridor[Idx], Verts);
 
-		CorridorPoly.SetColor(PolygonColor.WithAlpha(100));
-		CorridorPoly.Points.Reset();
-		CorridorPoly.Points.Append(Verts);
-		Snapshot->ElementsToDraw.Add(CorridorPoly);
+			FVector CenterPt = FVector::ZeroVector;
+			for (int32 VIdx = 0; VIdx < Verts.Num(); VIdx++)
+			{
+				Verts[VIdx].Z += 5.0f;
+				CenterPt += Verts[VIdx];
+			}
+			CenterPt /= Verts.Num();
 
-		if (AreaClass && AreaClass != UNavigationSystem::GetDefaultWalkableArea())
-		{
-			FVisualLogShapeElement AreaMarkElem(EVisualLoggerShapeElement::Segment);
-			AreaMarkElem.SetColor(FColorList::Orange.WithAlpha(100));
-			AreaMarkElem.Category = LogNavigation.GetCategoryName();
-			AreaMarkElem.Thicknes = 2;
-			AreaMarkElem.Description = AreaClass->GetName();
+			const UNavArea* DefArea = AreaClass ? ((UClass*)AreaClass)->GetDefaultObject<UNavArea>() : NULL;
+			const FColor PolygonColor = AreaClass != UNavigationSystem::GetDefaultWalkableArea() ? (DefArea ? DefArea->DrawColor : NavMesh->GetConfig().Color) : FColorList::LightSteelBlue;
 
-			AreaMarkElem.Points.Add(CenterPt + CorridorOffset);
-			AreaMarkElem.Points.Add(CenterPt + CorridorOffset + FVector(0, 0, 100.0f + NumAreaMark * 50.0f));
-			Snapshot->ElementsToDraw.Add(AreaMarkElem);
+			CorridorPoly.SetColor(PolygonColor.WithAlpha(100));
+			CorridorPoly.Points.Reset();
+			CorridorPoly.Points.Append(Verts);
+			Snapshot->ElementsToDraw.Add(CorridorPoly);
 
-			NumAreaMark = (NumAreaMark + 1) % 5;
+			if (AreaClass && AreaClass != UNavigationSystem::GetDefaultWalkableArea())
+			{
+				FVisualLogShapeElement AreaMarkElem(EVisualLoggerShapeElement::Segment);
+				AreaMarkElem.SetColor(FColorList::Orange.WithAlpha(100));
+				AreaMarkElem.Category = LogNavigation.GetCategoryName();
+				AreaMarkElem.Thicknes = 2;
+				AreaMarkElem.Description = AreaClass->GetName();
+
+				AreaMarkElem.Points.Add(CenterPt + CorridorOffset);
+				AreaMarkElem.Points.Add(CenterPt + CorridorOffset + FVector(0, 0, 100.0f + NumAreaMark * 50.0f));
+				Snapshot->ElementsToDraw.Add(AreaMarkElem);
+
+				NumAreaMark = (NumAreaMark + 1) % 5;
+			}
 		}
-	}
 
-	NavMesh->FinishBatchQuery();
+		NavMesh->FinishBatchQuery();
+	}
 #endif // ENABLE_VISUAL_LOG && WITH_RECAST
 }
 
diff --git a/Engine/Source/Runtime/AIModule/Private/Navigation/PathFollowingComponent.cpp b/Engine/Source/Runtime/AIModule/Private/Navigation/PathFollowingComponent.cpp
--- a/Engine/Source/Runtime/AIModule/Private/Navigation/PathFollowingComponent.cpp
+++ b/Engine/Source/Runtime/AIModule/Private/Navigation/PathFollowingComponent.cpp
@@ -70,17 +70,20 @@ void UPathFollowingComponent::LogPathHelper(const AActor* LogOwner, FNavigationP
 	if (Vlog.IsRecording() &&
 		InLogPath && InLogPath->IsValid() && InLogPath->GetPathPoints().Num())
 	{
-		FVisualLogEntry* Entry = Vlog.GetEntryToWrite(LogOwner, LogOwner->GetWorld()->TimeSeconds);
-		InLogPath->DescribeSelfToVisLog(Entry);
-
 		const FVector PathEnd = *InLogPath->GetPathPointLocation(InLogPath->GetPathPoints().Num() - 1);
-		if (LogGoalActor)
+
+		FVisualLogEntry* Entry = Vlog.GetEntryToWrite(LogOwner, LogOwner->GetWorld()->TimeSeconds);
+		if (Entry)
 		{
-			const FVector GoalLoc = LogGoalActor->GetActorLocation();
-			if (FVector::DistSquared(GoalLoc, PathEnd) > 1.0f)
+			InLogPath->DescribeSelfToVisLog(Entry);
+			if (LogGoalActor)
 			{
-				UE_VLOG_LOCATION(LogOwner, LogPathFollowing, Verbose, GoalLoc, 30, FColor::Green, TEXT("GoalActor"));
-				UE_VLOG_SEGMENT(LogOwner, LogPathFollowing, Verbose, GoalLoc, PathEnd, FColor::Green, TEXT_EMPTY);
+				const FVector GoalLoc = LogGoalActor->GetActorLocation();
+				if (FVector::DistSquared(GoalLoc, PathEnd) > 1.0f)
+				{
+					UE_VLOG_LOCATION(LogOwner, LogPathFollowing, Verbose, GoalLoc, 30, FColor::Green, TEXT("GoalActor"));
+					UE_VLOG_SEGMENT(LogOwner, LogPathFollowing, Verbose, GoalLoc, PathEnd, FColor::Green, TEXT_EMPTY);
+				}
 			}
 		}
 
diff --git a/Engine/Source/Runtime/AnimGraphRuntime/Private/AnimNodes/AnimNode_BlendListBase.cpp b/Engine/Source/Runtime/AnimGraphRuntime/Private/AnimNodes/AnimNode_BlendListBase.cpp
--- a/Engine/Source/Runtime/AnimGraphRuntime/Private/AnimNodes/AnimNode_BlendListBase.cpp
+++ b/Engine/Source/Runtime/AnimGraphRuntime/Private/AnimNodes/AnimNode_BlendListBase.cpp
@@ -207,8 +207,8 @@ void FAnimNode_BlendListBase::Evaluate(FPoseContext& Output)
 	if ((NumPoses > 0) && (BlendPose.Num() == BlendWeights.Num()))
 	{		
 		// Scratch arrays for evaluation, stack allocated
-		TArray<FCompactPose, TMemStackAllocator<ALIGNOF(FCompactPose)>> FilteredPoses;
-		TArray<FBlendedCurve, TMemStackAllocator<ALIGNOF(FBlendedCurve)>> FilteredCurve;
+		TArray<FCompactPose, TInlineAllocator<8>> FilteredPoses;
+		TArray<FBlendedCurve, TInlineAllocator<8>> FilteredCurve;
 		FilteredPoses.SetNum(NumPoses, false);
 		FilteredCurve.SetNum(NumPoses, false);
 
@@ -222,7 +222,7 @@ void FAnimNode_BlendListBase::Evaluate(FPoseContext& Output)
 			FPoseLink& CurrentPose = BlendPose[PoseIndex];
 			CurrentPose.Evaluate(EvaluateContext);
 
-			FilteredPoses[i].MoveBonesFrom(EvaluateContext.Pose);
+			FilteredPoses[i].CopyBonesFrom(EvaluateContext.Pose);
 			FilteredCurve[i] = EvaluateContext.Curve;
 		}
 
diff --git a/Engine/Source/Runtime/AnimGraphRuntime/Private/AnimNodes/AnimNode_LayeredBoneBlend.cpp b/Engine/Source/Runtime/AnimGraphRuntime/Private/AnimNodes/AnimNode_LayeredBoneBlend.cpp
--- a/Engine/Source/Runtime/AnimGraphRuntime/Private/AnimNodes/AnimNode_LayeredBoneBlend.cpp
+++ b/Engine/Source/Runtime/AnimGraphRuntime/Private/AnimNodes/AnimNode_LayeredBoneBlend.cpp
@@ -104,8 +104,8 @@ void FAnimNode_LayeredBoneBlend::Evaluate(FPoseContext& Output)
 				FPoseContext CurrentPoseContext(Output);
 				BlendPoses[ChildIndex].Evaluate(CurrentPoseContext);
 
-				TargetBlendPoses[ChildIndex].MoveBonesFrom(CurrentPoseContext.Pose);
-				TargetBlendCurves[ChildIndex].MoveFrom(CurrentPoseContext.Curve);
+				TargetBlendPoses[ChildIndex].CopyBonesFrom(CurrentPoseContext.Pose);
+				TargetBlendCurves[ChildIndex].CopyFrom(CurrentPoseContext.Curve);
 			}
 			else
 			{
diff --git a/Engine/Source/Runtime/AnimGraphRuntime/Private/AnimNodes/AnimNode_SaveCachedPose.cpp b/Engine/Source/Runtime/AnimGraphRuntime/Private/AnimNodes/AnimNode_SaveCachedPose.cpp
--- a/Engine/Source/Runtime/AnimGraphRuntime/Private/AnimNodes/AnimNode_SaveCachedPose.cpp
+++ b/Engine/Source/Runtime/AnimGraphRuntime/Private/AnimNodes/AnimNode_SaveCachedPose.cpp
@@ -62,8 +62,8 @@ void FAnimNode_SaveCachedPose::Evaluate(FPoseContext& Output)
 
 		FPoseContext CachingContext(Output);
 		Pose.Evaluate(CachingContext);
-		CachedPose.MoveBonesFrom(CachingContext.Pose);
-		CachedCurve.MoveFrom(CachingContext.Curve);
+		CachedPose.CopyBonesFrom(CachingContext.Pose);
+		CachedCurve.CopyFrom(CachingContext.Curve);
 	}
 
 	// Return the cached result
diff --git a/Engine/Source/Runtime/AnimGraphRuntime/Private/BoneControllers/AnimNode_SkeletalControlBase.cpp b/Engine/Source/Runtime/AnimGraphRuntime/Private/BoneControllers/AnimNode_SkeletalControlBase.cpp
--- a/Engine/Source/Runtime/AnimGraphRuntime/Private/BoneControllers/AnimNode_SkeletalControlBase.cpp
+++ b/Engine/Source/Runtime/AnimGraphRuntime/Private/BoneControllers/AnimNode_SkeletalControlBase.cpp
@@ -77,7 +77,7 @@ void FAnimNode_SkeletalControlBase::EvaluateComponentSpace(FComponentSpacePoseCo
 
 #if WITH_EDITORONLY_DATA
 			// save current pose before applying skeletal control to compute the exact gizmo location in AnimGraphNode
-			ForwardedPose = Output.Pose;
+			ForwardedPose.CopyPose(Output.Pose);
 #endif // #if WITH_EDITORONLY_DATA
 
 			BoneTransforms.Reset(BoneTransforms.Num());
diff --git a/Engine/Source/Runtime/AnimGraphRuntime/Public/AnimNodes/AnimNode_SaveCachedPose.h b/Engine/Source/Runtime/AnimGraphRuntime/Public/AnimNodes/AnimNode_SaveCachedPose.h
--- a/Engine/Source/Runtime/AnimGraphRuntime/Public/AnimNodes/AnimNode_SaveCachedPose.h
+++ b/Engine/Source/Runtime/AnimGraphRuntime/Public/AnimNodes/AnimNode_SaveCachedPose.h
@@ -13,8 +13,8 @@ struct ANIMGRAPHRUNTIME_API FAnimNode_SaveCachedPose : public FAnimNode_Base
 	FPoseLink Pose;
 
 protected:
-	FCompactPose CachedPose;
-	FBlendedCurve CachedCurve;
+	FCompactHeapPose CachedPose;
+	FBlendedHeapCurve CachedCurve;
 
 	FGraphTraversalCounter InitializationCounter;
 	FGraphTraversalCounter CachedBonesCounter;
diff --git a/Engine/Source/Runtime/AnimGraphRuntime/Public/BoneControllers/AnimNode_SkeletalControlBase.h b/Engine/Source/Runtime/AnimGraphRuntime/Public/BoneControllers/AnimNode_SkeletalControlBase.h
--- a/Engine/Source/Runtime/AnimGraphRuntime/Public/BoneControllers/AnimNode_SkeletalControlBase.h
+++ b/Engine/Source/Runtime/AnimGraphRuntime/Public/BoneControllers/AnimNode_SkeletalControlBase.h
@@ -48,7 +48,7 @@ public:
 public:
 #if WITH_EDITORONLY_DATA
 	// forwarded pose data from the wired node which current node's skeletal control is not applied yet
-	FCSPose<FCompactPose> ForwardedPose;
+	FCSPose<FCompactHeapPose> ForwardedPose;
 #endif //#if WITH_EDITORONLY_DATA
 
 	// FAnimNode_Base interface
diff --git a/Engine/Source/Runtime/Core/Private/Math/UnrealMath.cpp b/Engine/Source/Runtime/Core/Private/Math/UnrealMath.cpp
--- a/Engine/Source/Runtime/Core/Private/Math/UnrealMath.cpp
+++ b/Engine/Source/Runtime/Core/Private/Math/UnrealMath.cpp
@@ -2054,28 +2054,30 @@ FVector FMath::ClosestPointOnTetrahedronToPoint(const FVector& Point, const FVec
 	return Result;
 }
 
-void FMath::SphereDistToLine(FVector SphereOrigin, float SphereRadius, FVector LineOrigin, FVector LineDir, FVector& OutClosestPoint)
+void FMath::SphereDistToLine(FVector SphereOrigin, float SphereRadius, FVector LineOrigin, FVector NormalizedLineDir, FVector& OutClosestPoint)
 {
-	const float A	= LineDir | LineDir;
-	const float B	= 2.f * (LineDir | (LineOrigin - SphereOrigin));
-	const float C	= (SphereOrigin|SphereOrigin) + (LineOrigin|LineOrigin) - 2.f *(SphereOrigin|LineOrigin) - FMath::Square(SphereRadius);
-	const float D	= FMath::Square(B) - 4.f * A * C;
+	//const float A = NormalizedLineDir | NormalizedLineDir  (this is 1 because normalized)
+	//solving quadratic formula in terms of t where closest point = LineOrigin + t * NormalizedLineDir
+	const FVector LineOriginToSphereOrigin = SphereOrigin - LineOrigin;
+	const float B = -2.f * (NormalizedLineDir | LineOriginToSphereOrigin);
+	const float C = LineOriginToSphereOrigin.SizeSquared() - FMath::Square(SphereRadius);
+	const float D	= FMath::Square(B) - 4.f * C;
 
 	if( D <= KINDA_SMALL_NUMBER )
 	{
 		// line is not intersecting sphere (or is tangent at one point if D == 0 )
-		const FVector PointOnLine = LineOrigin + ( -B / 2.f * A ) * LineDir;
+		const FVector PointOnLine = LineOrigin + ( -B * 0.5f ) * NormalizedLineDir;
 		OutClosestPoint = SphereOrigin + (PointOnLine - SphereOrigin).GetSafeNormal() * SphereRadius;
 	}
 	else
 	{
 		// Line intersecting sphere in 2 points. Pick closest to line origin.
 		const float	E	= FMath::Sqrt(D);
-		const float T1	= (-B + E) / (2.f * A);
-		const float T2	= (-B - E) / (2.f * A);
-		const float T	= FMath::Abs(T1) < FMath::Abs(T2) ? T1 : T2;
+		const float T1	= (-B + E) * 0.5f;
+		const float T2	= (-B - E) * 0.5f;
+		const float T	= FMath::Abs( T1 ) == FMath::Abs( T2 ) ? FMath::Abs( T1 ) : FMath::Abs( T1 ) < FMath::Abs( T2 ) ? T1 : T2;	// In the case where both points are exactly the same distance we take the one in the direction of LineDir
 
-		OutClosestPoint	= LineOrigin + T * LineDir;
+		OutClosestPoint	= LineOrigin + T * NormalizedLineDir;
 	}
 }
 
diff --git a/Engine/Source/Runtime/Core/Public/Async/AsyncWork.h b/Engine/Source/Runtime/Core/Public/Async/AsyncWork.h
--- a/Engine/Source/Runtime/Core/Public/Async/AsyncWork.h
+++ b/Engine/Source/Runtime/Core/Public/Async/AsyncWork.h
@@ -252,7 +252,7 @@ class FAsyncTask
 	TTask Task;
 	/** Thread safe counter that indicates WORK completion, no necessarily finalization of the job */
 	FThreadSafeCounter	WorkNotFinishedCounter;
-	/** If we aren't doing the work synchrnously, this will hold the completion event */
+	/** If we aren't doing the work synchronously, this will hold the completion event */
 	FEvent*				DoneEvent;
 	/** Pool we are queued into, maintained by the calling thread */
 	FQueuedThreadPool*	QueuedPool;
diff --git a/Engine/Source/Runtime/Core/Public/Math/Vector2D.h b/Engine/Source/Runtime/Core/Public/Math/Vector2D.h
--- a/Engine/Source/Runtime/Core/Public/Math/Vector2D.h
+++ b/Engine/Source/Runtime/Core/Public/Math/Vector2D.h
@@ -323,7 +323,7 @@ public:
 	 * @param Tolerance Error tolerance.
 	 * @return true if the vectors are equal within specified tolerance, otherwise false.
 	 */
-	bool Equals( const FVector2D& V, float Tolerance ) const;
+	bool Equals( const FVector2D& V, float Tolerance=KINDA_SMALL_NUMBER ) const;
 
 	/**
 	 * Set the values of the vector directly.
diff --git a/Engine/Source/Runtime/Engine/Classes/Animation/AnimInstance.h b/Engine/Source/Runtime/Engine/Classes/Animation/AnimInstance.h
--- a/Engine/Source/Runtime/Engine/Classes/Animation/AnimInstance.h
+++ b/Engine/Source/Runtime/Engine/Classes/Animation/AnimInstance.h
@@ -200,9 +200,9 @@ struct FSlotEvaluationPose
 	UPROPERTY()
 	float Weight;
 
-	/** Pose */
+	/*** ATTENTION *****/
+	/* These Pose/Curve is stack allocator. You should not use it outside of stack. */
 	FCompactPose Pose;
-	/** Curve */
 	FBlendedCurve Curve;
 
 	FSlotEvaluationPose()
@@ -1012,7 +1012,7 @@ public:
 	bool ParallelCanEvaluate(const USkeletalMesh* InSkeletalMesh) const;
 
 	/** Perform evaluation. Can be called from worker threads. */
-	void ParallelEvaluateAnimation(bool bForceRefPose, const USkeletalMesh* InSkeletalMesh, TArray<FTransform>& OutLocalAtoms, TArray<FActiveVertexAnim>& OutVertexAnims, FBlendedCurve& OutCurve);
+	void ParallelEvaluateAnimation(bool bForceRefPose, const USkeletalMesh* InSkeletalMesh, TArray<FTransform>& OutLocalAtoms, TArray<FActiveVertexAnim>& OutVertexAnims, FBlendedHeapCurve& OutCurve);
 
 	void PostEvaluateAnimation();
 	void UninitializeAnimation();
@@ -1099,7 +1099,7 @@ private:
 
 public: 
 	/** Update all internal curves from Blended Curve */
-	void UpdateCurves(const FBlendedCurve& InCurves);
+	void UpdateCurves(const FBlendedHeapCurve& InCurves);
 
 	/** Check whether we have active morph target curves */
 	bool HasMorphTargetCurves() const;
diff --git a/Engine/Source/Runtime/Engine/Classes/Animation/AnimNodeBase.h b/Engine/Source/Runtime/Engine/Classes/Animation/AnimNodeBase.h
--- a/Engine/Source/Runtime/Engine/Classes/Animation/AnimNodeBase.h
+++ b/Engine/Source/Runtime/Engine/Classes/Animation/AnimNodeBase.h
@@ -131,6 +131,7 @@ public:
 struct FPoseContext : public FAnimationBaseContext
 {
 public:
+	/* These Pose/Curve is stack allocator. You should not use it outside of stack. */
 	FCompactPose	Pose;
 	FBlendedCurve	Curve;
 
@@ -387,7 +388,7 @@ public:
 #if ENABLE_ANIMNODE_POSE_DEBUG
 private:
 	// forwarded pose data from the wired node which current node's skeletal control is not applied yet
-	FCompactPose CurrentPose;
+	FCompactHeapPose CurrentPose;
 #endif //#if ENABLE_ANIMNODE_POSE_DEBUG
 };
 
diff --git a/Engine/Source/Runtime/Engine/Classes/Animation/AnimNode_TransitionPoseEvaluator.h b/Engine/Source/Runtime/Engine/Classes/Animation/AnimNode_TransitionPoseEvaluator.h
--- a/Engine/Source/Runtime/Engine/Classes/Animation/AnimNode_TransitionPoseEvaluator.h
+++ b/Engine/Source/Runtime/Engine/Classes/Animation/AnimNode_TransitionPoseEvaluator.h
@@ -51,8 +51,8 @@ struct ENGINE_API FAnimNode_TransitionPoseEvaluator : public FAnimNode_Base
 	UPROPERTY(EditAnywhere, BlueprintReadWrite, Category=Pose, meta=(NeverAsPin, ClampMin="1", UIMin="1"))
 	int32 FramesToCachePose;
 
-	FCompactPose CachedPose;
-	FBlendedCurve CachedCurve;
+	FCompactHeapPose CachedPose;
+	FBlendedHeapCurve CachedCurve;
 
 	UPROPERTY(transient)
 	int32 CacheFramesRemaining;
diff --git a/Engine/Source/Runtime/Engine/Classes/Animation/AnimStateMachineTypes.h b/Engine/Source/Runtime/Engine/Classes/Animation/AnimStateMachineTypes.h
--- a/Engine/Source/Runtime/Engine/Classes/Animation/AnimStateMachineTypes.h
+++ b/Engine/Source/Runtime/Engine/Classes/Animation/AnimStateMachineTypes.h
@@ -93,7 +93,7 @@ struct FAnimationState : public FAnimationStateBase
 	// The index of the notify to fire when this state is fully entered (weight within the machine becomes one)
 	UPROPERTY()
 	int32 FullyBlendedNotify;
-
+	
 	FAnimationState()
 		: FAnimationStateBase()
 		, StartNotify(INDEX_NONE)
@@ -239,6 +239,10 @@ struct FBakedAnimationState
 	UPROPERTY()
 	TArray<int32> PlayerNodeIndices;
 
+	// Whether or not this state will ALWAYS reset it's state on reentry, regardless of remaining weight
+	UPROPERTY()
+	bool bAlwaysResetOnEntry;
+
 public:
 	FBakedAnimationState()
 		: StartNotify(INDEX_NONE)
diff --git a/Engine/Source/Runtime/Engine/Classes/Components/ChildActorComponent.h b/Engine/Source/Runtime/Engine/Classes/Components/ChildActorComponent.h
--- a/Engine/Source/Runtime/Engine/Classes/Components/ChildActorComponent.h
+++ b/Engine/Source/Runtime/Engine/Classes/Components/ChildActorComponent.h
@@ -3,18 +3,42 @@
 
 #pragma once
 #include "SceneComponent.h"
+#include "ComponentInstanceDataCache.h"
 #include "ChildActorComponent.generated.h"
 
-class FChildActorComponentInstanceData;
+class ENGINE_API FChildActorComponentInstanceData : public FSceneComponentInstanceData
+{
+public:
+	FChildActorComponentInstanceData(const class UChildActorComponent* Component);
+
+	virtual ~FChildActorComponentInstanceData();
+
+	virtual void ApplyToComponent(UActorComponent* Component, const ECacheApplyPhase CacheApplyPhase) override;
+
+	// The name of the spawned child actor so it (attempts to) remain constant across construction script reruns
+	FName ChildActorName;
+
+	// The component instance data cache for the ChildActor spawned by this component
+	FComponentInstanceDataCache* ComponentInstanceData;
+
+	struct FAttachedActorInfo
+	{
+		TWeakObjectPtr<AActor> Actor;
+		FName SocketName;
+		FTransform RelativeTransform;
+	};
+
+	TArray<FAttachedActorInfo> AttachedActors;
+};
 
 /** A component that spawns an Actor when registered, and destroys it when unregistered.*/
-UCLASS(ClassGroup=Utility, hidecategories=(Object,LOD,Physics,Lighting,TextureStreaming,Activation,"Components|Activation",Collision), meta=(BlueprintSpawnableComponent), MinimalAPI)
-class UChildActorComponent : public USceneComponent
+UCLASS(ClassGroup=Utility, hidecategories=(Object,LOD,Physics,Lighting,TextureStreaming,Activation,"Components|Activation",Collision), meta=(BlueprintSpawnableComponent))
+class ENGINE_API UChildActorComponent : public USceneComponent
 {
 	GENERATED_UCLASS_BODY()
 
 	UFUNCTION(BlueprintCallable, Category=ChildActorComponent)
-	ENGINE_API void SetChildActorClass(TSubclassOf<AActor> InClass);
+	void SetChildActorClass(TSubclassOf<AActor> InClass);
 
 	TSubclassOf<AActor> GetChildActorClass() const { return ChildActorClass; }
 
@@ -55,7 +79,7 @@ public:
 	void ApplyComponentInstanceData(class FChildActorComponentInstanceData* ComponentInstanceData, const ECacheApplyPhase CacheApplyPhase);
 
 	/** Create the child actor */
-	ENGINE_API void CreateChildActor();
+	void CreateChildActor();
 
 	/** Kill any currently present child actor */
 	void DestroyChildActor(const bool bRequiresRename = true);
diff --git a/Engine/Source/Runtime/Engine/Classes/Components/SkeletalMeshComponent.h b/Engine/Source/Runtime/Engine/Classes/Components/SkeletalMeshComponent.h
--- a/Engine/Source/Runtime/Engine/Classes/Components/SkeletalMeshComponent.h
+++ b/Engine/Source/Runtime/Engine/Classes/Components/SkeletalMeshComponent.h
@@ -108,7 +108,7 @@ struct FAnimationEvaluationContext
 	FVector RootBoneTranslation;
 
 	// Double buffer curve data
-	FBlendedCurve	Curve;
+	FBlendedHeapCurve	Curve;
 
 	// Are we performing interpolation this tick
 	bool bDoInterpolation;
@@ -380,7 +380,7 @@ public:
 	TArray<FTransform> CachedSpaceBases;
 
 	/** Cached Curve result for Update Rate optimization */
-	FBlendedCurve CachedCurve;
+	FBlendedHeapCurve CachedCurve;
 
 	/** Used to scale speed of all animations on this skeletal mesh. */
 	UPROPERTY(EditAnywhere, AdvancedDisplay, BlueprintReadWrite, Category=Animation)
@@ -410,6 +410,7 @@ public:
 	/**
 	 *  If true, simulate physics for this component on a dedicated server.
 	 *  This should be set if simulating physics and replicating with a dedicated server.
+	 *	Note: This property cannot be changed at runtime.
 	 */
 	UPROPERTY(EditAnywhere, AdvancedDisplay, BlueprintReadWrite, Category = SkeletalMesh)
 	uint32 bEnablePhysicsOnDedicatedServer:1;
@@ -1029,7 +1030,7 @@ public:
 	* @param	OutVertexAnims			Active vertex animations
 	* @param	OutRootBoneTranslation	Calculated root bone translation
 	*/
-	void PerformAnimationEvaluation(const USkeletalMesh* InSkeletalMesh, UAnimInstance* InAnimInstance, TArray<FTransform>& OutSpaceBases, TArray<FTransform>& OutLocalAtoms, TArray<FActiveVertexAnim>& OutVertexAnims, FVector& OutRootBoneTranslation, FBlendedCurve& OutCurve) const;
+	void PerformAnimationEvaluation(const USkeletalMesh* InSkeletalMesh, UAnimInstance* InAnimInstance, TArray<FTransform>& OutSpaceBases, TArray<FTransform>& OutLocalAtoms, TArray<FActiveVertexAnim>& OutVertexAnims, FVector& OutRootBoneTranslation, FBlendedHeapCurve& OutCurve) const;
 	void PostAnimEvaluation( FAnimationEvaluationContext& EvaluationContext );
 
 	/**
@@ -1325,7 +1326,7 @@ private:
 	void PostPhysicsTickComponent(FSkeletalMeshComponentPostPhysicsTickFunction& ThisTickFunction);
 
 	/** Evaluate Anim System **/
-	void EvaluateAnimation(const USkeletalMesh* InSkeletalMesh, UAnimInstance* InAnimInstance, TArray<FTransform>& OutLocalAtoms, TArray<struct FActiveVertexAnim>& OutVertexAnims, FVector& OutRootBoneTranslation, FBlendedCurve& OutCurve) const;
+	void EvaluateAnimation(const USkeletalMesh* InSkeletalMesh, UAnimInstance* InAnimInstance, TArray<FTransform>& OutLocalAtoms, TArray<struct FActiveVertexAnim>& OutVertexAnims, FVector& OutRootBoneTranslation, FBlendedHeapCurve& OutCurve) const;
 
 	/**
 	* Take the SourceAtoms array (translation vector, rotation quaternion and scale vector) and update the array of component-space bone transformation matrices (DestSpaceBases).
diff --git a/Engine/Source/Runtime/Engine/Classes/Components/TextRenderComponent.h b/Engine/Source/Runtime/Engine/Classes/Components/TextRenderComponent.h
--- a/Engine/Source/Runtime/Engine/Classes/Components/TextRenderComponent.h
+++ b/Engine/Source/Runtime/Engine/Classes/Components/TextRenderComponent.h
@@ -81,6 +81,10 @@ class ENGINE_API UTextRenderComponent : public UPrimitiveComponent
 	UPROPERTY(EditAnywhere, BlueprintReadOnly, AdvancedDisplay, Category=Text)
 	float HorizSpacingAdjust;
 
+	/** Vertical adjustment per character, default is 0.0 */
+	UPROPERTY(EditAnywhere, BlueprintReadOnly, AdvancedDisplay, Category=Text)
+	float VertSpacingAdjust;
+
 	/** Allows text to draw unmodified when using debug visualization modes. **/
 	UPROPERTY(EditAnywhere, BlueprintReadOnly, AdvancedDisplay, Category=Rendering)
 	uint32 bAlwaysRenderAsText:1;
@@ -130,6 +134,10 @@ class ENGINE_API UTextRenderComponent : public UPrimitiveComponent
 	UFUNCTION(BlueprintCallable, Category="Rendering|Components|TextRender")
 	void SetHorizSpacingAdjust(float Value);
 
+	/** Change the text vertical spacing adjustment and signal the primitives to be rebuilt */
+	UFUNCTION(BlueprintCallable, Category="Rendering|Components|TextRender")
+	void SetVertSpacingAdjust(float Value);
+
 	/** Change the world size of the text and signal the primitives to be rebuilt */
 	UFUNCTION(BlueprintCallable, Category="Rendering|Components|TextRender")
 	void SetWorldSize(float Value);
diff --git a/Engine/Source/Runtime/Engine/Classes/GameFramework/Actor.h b/Engine/Source/Runtime/Engine/Classes/GameFramework/Actor.h
--- a/Engine/Source/Runtime/Engine/Classes/GameFramework/Actor.h
+++ b/Engine/Source/Runtime/Engine/Classes/GameFramework/Actor.h
@@ -2398,6 +2398,16 @@ public:
 	UFUNCTION(BlueprintImplementableEvent, Category=Actor, meta=(DisplayName="OnReset"))
 	void K2_OnReset();
 
+	/**
+	 * Returns true if this actor has been rendered "recently", with a tolerance in seconds to define what "recent" means. 
+	 * e.g.: If a tolerance of 0.1 is used, this function will return true only if the actor was rendered in the last 0.1 seconds of game time. 
+	 *
+	 * @param Tolerance  How many seconds ago the actor last render time can be and still count as having been "recently" rendered.
+	 * @return Whether this actor was recently rendered.
+	 */
+	UFUNCTION(Category="Rendering", BlueprintCallable, meta=(Keywords="scene visible"))
+	bool WasRecentlyRendered(float Tolerance = 0.2) const;
+
 	/** Returns the most recent time any of this actor's components were rendered */
 	virtual float GetLastRenderTime() const;
 
diff --git a/Engine/Source/Runtime/Engine/Classes/GameFramework/Controller.h b/Engine/Source/Runtime/Engine/Classes/GameFramework/Controller.h
--- a/Engine/Source/Runtime/Engine/Classes/GameFramework/Controller.h
+++ b/Engine/Source/Runtime/Engine/Classes/GameFramework/Controller.h
@@ -92,6 +92,10 @@ protected:
 	UPROPERTY(EditDefaultsOnly, AdvancedDisplay, Category="Controller|Transform")
 	uint32 bAttachToPawn:1;
 
+	/** Whether this controller is a player controller. */
+	UPROPERTY()
+	uint32 bIsPlayerController:1;
+
 	/**
 	  * Physically attach the Controller to the specified Pawn, so that our position reflects theirs.
 	  * The attachment persists during possession of the pawn. The Controller's rotation continues to match the ControlRotation.
@@ -209,7 +213,10 @@ public:
 
 	/** returns whether this Controller is a locally controlled PlayerController.  */
 	UFUNCTION(BlueprintCallable, Category="Pawn")
-	virtual bool IsLocalPlayerController() const;
+	bool IsLocalPlayerController() const
+	{
+		return bIsPlayerController && IsLocalController();
+	}
 
 	/** Returns whether this Controller is a local controller.	 */
 	UFUNCTION(BlueprintCallable, Category="Pawn")
diff --git a/Engine/Source/Runtime/Engine/Classes/GameFramework/GameMode.h b/Engine/Source/Runtime/Engine/Classes/GameFramework/GameMode.h
--- a/Engine/Source/Runtime/Engine/Classes/GameFramework/GameMode.h
+++ b/Engine/Source/Runtime/Engine/Classes/GameFramework/GameMode.h
@@ -125,6 +125,9 @@ protected:
 	/** Updates the match state and calls the appropriate transition functions */
 	virtual void SetMatchState(FName NewState);
 
+	/** Overridable virtual function to dispatch the appropriate transition functions before GameState and Blueprints get SetMatchState calls. */
+	virtual void OnMatchStateSet();
+
 	/** Implementable event to respond to match state changes */
 	UFUNCTION(BlueprintImplementableEvent, Category="Game", meta=(DisplayName="OnSetMatchState"))
 	void K2_OnSetMatchState(FName NewState);
@@ -634,10 +637,10 @@ public:
 	 * only dynamic actors in the PersistentLevel may be moved (this includes all actors spawned during gameplay)
 	 * this is called for both parts of the transition because actors might change while in the middle (e.g. players might join or leave the game)
 	 * @see also PlayerController::GetSeamlessTravelActorList() (the function that's called on clients)
-	 * @param bToEntry true if we are going from old level -> entry, false if we are going from entry -> new level
+	 * @param bToTransition true if we are going from old level to transition map, false if we are going from transition map to new level
 	 * @param ActorList (out) list of actors to maintain
 	 */
-	virtual void GetSeamlessTravelActorList(bool bToEntry, TArray<AActor*>& ActorList);
+	virtual void GetSeamlessTravelActorList(bool bToTransition, TArray<AActor*>& ActorList);
 
 	/** Allow the game to specify a place for clients to download MapName */
 	DEPRECATED(4.11, "Override GameWelcomePlayer() instead")
diff --git a/Engine/Source/Runtime/Engine/Classes/GameFramework/HUD.h b/Engine/Source/Runtime/Engine/Classes/GameFramework/HUD.h
--- a/Engine/Source/Runtime/Engine/Classes/GameFramework/HUD.h
+++ b/Engine/Source/Runtime/Engine/Classes/GameFramework/HUD.h
@@ -236,9 +236,10 @@ public:
 	 * @param EndScreenX		Screen-space X coordinate of end of the line.
 	 * @param EndScreenY		Screen-space Y coordinate of end of the line.
 	 * @param LineColor			Color to draw line
+	 * @param LineThickness		Thickness of the line to draw
 	 */
 	UFUNCTION(BlueprintCallable, Category=HUD, meta=(LineColor="(R=0,G=0,B=0,A=1)"))
-	void DrawLine(float StartScreenX, float StartScreenY, float EndScreenX, float EndScreenY, FLinearColor LineColor);
+	void DrawLine(float StartScreenX, float StartScreenY, float EndScreenX, float EndScreenY, FLinearColor LineColor, float LineThickness=0.f);
 
 	/**
 	 * Draws a colored untextured quad on the HUD.
diff --git a/Engine/Source/Runtime/Engine/Classes/GameFramework/PlayerController.h b/Engine/Source/Runtime/Engine/Classes/GameFramework/PlayerController.h
--- a/Engine/Source/Runtime/Engine/Classes/GameFramework/PlayerController.h
+++ b/Engine/Source/Runtime/Engine/Classes/GameFramework/PlayerController.h
@@ -1268,7 +1268,6 @@ public:
 
 	//~ Begin AController Interface
 	virtual void GameHasEnded(class AActor* EndGameFocus = NULL, bool bIsWinner = false) override;
-	virtual bool IsLocalPlayerController() const override;
 	virtual bool IsLocalController() const override;
 	virtual void GetPlayerViewPoint(FVector& out_Location, FRotator& out_Rotation) const override;
 	virtual void SetInitialLocationAndRotation(const FVector& NewLocation, const FRotator& NewRotation) override;
@@ -1622,7 +1621,7 @@ private:
 
 	/** Set during SpawnActor once and never again to indicate the intent of this controller instance (SERVER ONLY) */
 	UPROPERTY()
-	bool		bIsLocalPlayerController;
+	mutable bool bIsLocalPlayerController;
 
 public:
 	/** Counter for this players seamless travels (used along with the below value, to restrict ServerNotifyLoadedWorld) */
diff --git a/Engine/Source/Runtime/Engine/Classes/Kismet/GameplayStatics.h b/Engine/Source/Runtime/Engine/Classes/Kismet/GameplayStatics.h
--- a/Engine/Source/Runtime/Engine/Classes/Kismet/GameplayStatics.h
+++ b/Engine/Source/Runtime/Engine/Classes/Kismet/GameplayStatics.h
@@ -323,6 +323,23 @@ class ENGINE_API UGameplayStatics : public UBlueprintFunctionLibrary
 	UFUNCTION(BlueprintCallable, BlueprintCosmetic, Category = "Audio", meta = (WorldContext = "WorldContextObject"))
 	static void SetGlobalPitchModulation(UObject* WorldContextObject, float PitchModulation, float TimeSec);
 
+	/**
+	* Sets the global listener focus parameters which will scale focus behavior of sounds based on their focus azimuth settings in their attenuation settings. 
+	*
+	* * Fire and Forget.
+	* * Not Replicated.
+	* @param FocusAzimuthScale - An angle scale value used to scale the azimuth angle that defines where sounds are in-focus.
+	* @param NonFocusAzimuthScale- An angle scale value used to scale the azimuth angle that defines where sounds are out-of-focus.
+	* @param FocusDistanceScale - A distance scale value to use for sounds which are in-focus. Values < 1.0 will reduce perceived distance to sounds, values > 1.0 will increase perceived distance to in-focus sounds.
+	* @param NonFocusDistanceScale - A distance scale value to use for sounds which are out-of-focus. Values < 1.0 will reduce perceived distance to sounds, values > 1.0 will increase perceived distance to in-focus sounds.
+	* @param FocusVolumeScale- A volume attenuation value to use for sounds which are in-focus.
+	* @param NonFocusVolumeScale- A volume attenuation value to use for sounds which are out-of-focus.
+	* @param FocusPriorityScale - A priority scale value (> 0.0) to use for sounds which are in-focus. Values < 1.0 will reduce the priority of in-focus sounds, values > 1.0 will increase the priority of in-focus sounds.
+	* @param NonFocusPriorityScale - A priority scale value (> 0.0) to use for sounds which are out-of-focus. Values < 1.0 will reduce the priority of sounds out-of-focus sounds, values > 1.0 will increase the priority of out-of-focus sounds.
+	*/
+	UFUNCTION(BlueprintCallable, BlueprintCosmetic, Category = "Audio", meta = (WorldContext = "WorldContextObject"))
+	static void SetGlobalListenerFocusParameters(UObject* WorldContextObject, float FocusAzimuthScale = 1.0f, float NonFocusAzimuthScale = 1.0f, float FocusDistanceScale = 1.0f, float NonFocusDistanceScale = 1.0f, float FocusVolumeScale = 1.0f, float NonFocusVolumeScale = 1.0f, float FocusPriorityScale = 1.0f, float NonFocusPriorityScale = 1.0f);
+
 	/**
 	 * Plays a sound directly with no attenuation, perfect for UI sounds.
 	 *
diff --git a/Engine/Source/Runtime/Engine/Classes/Kismet/KismetSystemLibrary.h b/Engine/Source/Runtime/Engine/Classes/Kismet/KismetSystemLibrary.h
--- a/Engine/Source/Runtime/Engine/Classes/Kismet/KismetSystemLibrary.h
+++ b/Engine/Source/Runtime/Engine/Classes/Kismet/KismetSystemLibrary.h
@@ -1251,38 +1251,38 @@ class ENGINE_API UKismetSystemLibrary : public UBlueprintFunctionLibrary
 
 	/** Draw directional arrow, pointing from LineStart to LineEnd. */
 	UFUNCTION(BlueprintCallable, Category="Rendering|Debug", meta=(WorldContext="WorldContextObject"))
-	static void DrawDebugArrow(UObject* WorldContextObject, const FVector LineStart, const FVector LineEnd, float ArrowSize, FLinearColor LineColor, float Duration=0.f);
+	static void DrawDebugArrow(UObject* WorldContextObject, const FVector LineStart, const FVector LineEnd, float ArrowSize, FLinearColor LineColor, float Duration=0.f, float Thickness = 0.f);
 
 	/** Draw a debug box */
 	UFUNCTION(BlueprintCallable, Category="Rendering|Debug", meta=(WorldContext="WorldContextObject"))
-	static void DrawDebugBox(UObject* WorldContextObject, const FVector Center, FVector Extent, FLinearColor LineColor, const FRotator Rotation=FRotator::ZeroRotator, float Duration=0.f);
+	static void DrawDebugBox(UObject* WorldContextObject, const FVector Center, FVector Extent, FLinearColor LineColor, const FRotator Rotation=FRotator::ZeroRotator, float Duration=0.f, float Thickness = 0.f);
 
 	/** Draw a debug coordinate system. */
 	UFUNCTION(BlueprintCallable, Category="Rendering|Debug", meta=(WorldContext="WorldContextObject"))
-	static void DrawDebugCoordinateSystem(UObject* WorldContextObject, const FVector AxisLoc, const FRotator AxisRot, float Scale=1.f, float Duration=0.f);
+	static void DrawDebugCoordinateSystem(UObject* WorldContextObject, const FVector AxisLoc, const FRotator AxisRot, float Scale=1.f, float Duration=0.f, float Thickness = 0.f);
 
 	/** Draw a debug sphere */
 	UFUNCTION(BlueprintCallable, Category="Rendering|Debug", meta=(WorldContext="WorldContextObject"))
-	static void DrawDebugSphere(UObject* WorldContextObject, const FVector Center, float Radius=100.f, int32 Segments=12, FLinearColor LineColor = FLinearColor::White, float Duration=0.f);
+	static void DrawDebugSphere(UObject* WorldContextObject, const FVector Center, float Radius=100.f, int32 Segments=12, FLinearColor LineColor = FLinearColor::White, float Duration=0.f, float Thickness = 0.f);
 
 	/** Draw a debug cylinder */
 	UFUNCTION(BlueprintCallable, Category="Rendering|Debug", meta=(WorldContext="WorldContextObject"))
-	static void DrawDebugCylinder(UObject* WorldContextObject, const FVector Start, const FVector End, float Radius=100.f, int32 Segments=12, FLinearColor LineColor = FLinearColor::White, float Duration=0.f);
+	static void DrawDebugCylinder(UObject* WorldContextObject, const FVector Start, const FVector End, float Radius=100.f, int32 Segments=12, FLinearColor LineColor = FLinearColor::White, float Duration=0.f, float Thickness = 0.f);
 	
 	/** Draw a debug cone */
 	UFUNCTION(BlueprintCallable, Category="Rendering|Debug", meta=(WorldContext="WorldContextObject", DeprecatedFunction, DeprecationMessage="DrawDebugCone has been changed to use degrees for angles instead of radians. Place a new DrawDebugCone node and pass your angles as degrees."))
-	static void DrawDebugCone(UObject* WorldContextObject, const FVector Origin, const FVector Direction, float Length, float AngleWidth, float AngleHeight, int32 NumSides, FLinearColor LineColor);
+	static void DrawDebugCone(UObject* WorldContextObject, const FVector Origin, const FVector Direction, float Length, float AngleWidth, float AngleHeight, int32 NumSides, FLinearColor LineColor, float Duration = 0.f, float Thickness = 0.f);
 
 	/** 
 	 * Draw a debug cone 
 	 * Angles are specified in degrees
 	 */
 	UFUNCTION(BlueprintCallable, Category="Rendering|Debug", meta=(WorldContext="WorldContextObject", DisplayName="DrawDebugCone"))
-	static void DrawDebugConeInDegrees(UObject* WorldContextObject, const FVector Origin, const FVector Direction, float Length=100.f, float AngleWidth=45.f, float AngleHeight=45.f, int32 NumSides = 12, FLinearColor LineColor = FLinearColor::White, float Duration=0.f);
+	static void DrawDebugConeInDegrees(UObject* WorldContextObject, const FVector Origin, const FVector Direction, float Length=100.f, float AngleWidth=45.f, float AngleHeight=45.f, int32 NumSides = 12, FLinearColor LineColor = FLinearColor::White, float Duration=0.f, float Thickness = 0.f);
 
 	/** Draw a debug capsule */
 	UFUNCTION(BlueprintCallable, Category="Rendering|Debug", meta=(WorldContext="WorldContextObject"))
-	static void DrawDebugCapsule(UObject* WorldContextObject, const FVector Center, float HalfHeight, float Radius, const FRotator Rotation, FLinearColor LineColor = FLinearColor::White, float Duration=0.f);
+	static void DrawDebugCapsule(UObject* WorldContextObject, const FVector Center, float HalfHeight, float Radius, const FRotator Rotation, FLinearColor LineColor = FLinearColor::White, float Duration=0.f, float Thickness = 0.f);
 
 	/** Draw a debug string at a 3d world location. */
 	UFUNCTION(BlueprintCallable, Category="Rendering|Debug", meta=(WorldContext="WorldContextObject"))
@@ -1309,7 +1309,7 @@ class ENGINE_API UKismetSystemLibrary : public UBlueprintFunctionLibrary
 
 	/** Draws a debug frustum. */
 	UFUNCTION(BlueprintCallable, Category="Rendering|Debug", meta=(WorldContext="WorldContextObject"))
-	static void DrawDebugFrustum(UObject* WorldContextObject, const FTransform& FrustumTransform, FLinearColor FrustumColor = FLinearColor::White, float Duration=0.f);
+	static void DrawDebugFrustum(UObject* WorldContextObject, const FTransform& FrustumTransform, FLinearColor FrustumColor = FLinearColor::White, float Duration=0.f, float Thickness = 0.f);
 
 	/** Draw a debug camera shape. */
 	UFUNCTION(BlueprintCallable, Category="Rendering|Debug")
diff --git a/Engine/Source/Runtime/Engine/Classes/Particles/ParticleSystem.h b/Engine/Source/Runtime/Engine/Classes/Particles/ParticleSystem.h
--- a/Engine/Source/Runtime/Engine/Classes/Particles/ParticleSystem.h
+++ b/Engine/Source/Runtime/Engine/Classes/Particles/ParticleSystem.h
@@ -345,6 +345,7 @@ public:
 	virtual SIZE_T GetResourceSize(EResourceSizeMode::Type Mode) override;
 	virtual void GetAssetRegistryTags(TArray<FAssetRegistryTag>& OutTags) const override;
 	bool IsPotentiallyImmortal() const;
+	bool UsesCPUCollision() const;
 	virtual bool CanBeClusterRoot() const override;
 	//~ End UObject Interface.
 
diff --git a/Engine/Source/Runtime/Engine/Classes/Particles/ParticleSystemComponent.h b/Engine/Source/Runtime/Engine/Classes/Particles/ParticleSystemComponent.h
--- a/Engine/Source/Runtime/Engine/Classes/Particles/ParticleSystemComponent.h
+++ b/Engine/Source/Runtime/Engine/Classes/Particles/ParticleSystemComponent.h
@@ -121,7 +121,6 @@ struct FParticleSysParam
 	UPROPERTY(EditAnywhere, BlueprintReadWrite, Category=ParticleSysParam)
 	class UMaterialInterface* Material;
 
-
 	FParticleSysParam()
 		: ParamType(0)
 		, Scalar(0)
@@ -131,9 +130,65 @@ struct FParticleSysParam
 		, Color(ForceInit)
 		, Actor(NULL)
 		, Material(NULL)
+		, AsyncActorToWorld(FTransform::Identity)
+		, AsyncActorVelocity(FVector::ZeroVector)
+		, bAsyncDataCopyIsValid(false)
+	{
+	}
+
+	void UpdateAsyncActorCache()
 	{
+		check(IsInGameThread());
+		if(Actor)
+		{
+			AsyncActorToWorld = Actor->ActorToWorld();
+			AsyncActorVelocity = Actor->GetVelocity();
+		}
+		
+		bAsyncDataCopyIsValid = true;
+	}
+
+	void ResetAsyncActorCache()
+	{
+		check(IsInGameThread());
+		bAsyncDataCopyIsValid = false;
+	}
+
+	FTransform GetAsyncActorToWorld() const
+	{
+		if(bAsyncDataCopyIsValid)
+		{
+			return AsyncActorToWorld;
+		}
+		else if(Actor)
+		{
+			check(IsInGameThread());
+			return Actor->ActorToWorld();
+		}
+
+		return FTransform::Identity;
+	}
+
+	FVector GetAsyncActorVelocity() const
+	{
+		if (bAsyncDataCopyIsValid)
+		{
+			return AsyncActorVelocity;
+		}
+		else if (Actor)
+		{
+			check(IsInGameThread());
+			return Actor->GetVelocity();
+		}
+
+		return FVector::ZeroVector;
 	}
 
+private:
+	FTransform AsyncActorToWorld;
+	FVector AsyncActorVelocity;
+	bool bAsyncDataCopyIsValid;
+
 };
 
 template <> struct TIsPODType<FParticleSysParam> { enum { Value = true }; };
@@ -536,9 +591,22 @@ private:
 	FTransform AsyncComponentToWorld;
 	/** Cached copy of the instance params */
 	TArray<struct FParticleSysParam> AsyncInstanceParameters;
+	/** Player locations computed before kicking off async task. Safe to access from async task or game thread*/
+	TArray<FVector> PlayerLocations;
+	/** PlayerLODDistanceFactor computed before kicking off async task. Safe to access from async task or game thread*/
+	TArray<float> PlayerLODDistanceFactor;
+	/** Cached copy of bounds */
+	FBoxSphereBounds AsyncBounds;
+	/** Cached copy of PartSysVelocity */
+	FVector AsyncPartSysVelocity;
+
 	/** Is AsyncComponentToWorld etc valid? */
 	bool bAsyncDataCopyIsValid;
 	bool bParallelRenderThreadUpdate;
+
+	/** Remember the global detail mode when we last checked the emitters. */
+	uint32 LastCheckedDetailMode;
+
 public:
 
 	void SetSignificance(float OldSignificance, float NewSignificance);
@@ -563,6 +631,36 @@ public:
 		return InstanceParameters;
 	}
 
+	FORCEINLINE const FBoxSphereBounds& GetAsyncBounds()
+	{
+		if (!bParallelRenderThreadUpdate && !IsInGameThread())
+		{
+			check(bAsyncDataCopyIsValid);
+			return AsyncBounds;
+		}
+		return Bounds;
+	}
+
+	FORCEINLINE const FVector& GetAsyncPartSysVelocity()
+	{
+		if (!bParallelRenderThreadUpdate && !IsInGameThread())
+		{
+			check(bAsyncDataCopyIsValid);
+			return AsyncPartSysVelocity;
+		}
+		return PartSysVelocity;
+	}
+
+	FORCEINLINE const TArray<FVector>& GetPlayerLocations()
+	{
+		return PlayerLocations;
+	}
+
+	FORCEINLINE const TArray<float>& GetPlayerLODDistanceFactor()
+	{
+		return PlayerLODDistanceFactor;
+	}
+
 	//
 	//	Beam-related script functions.
 	//
diff --git a/Engine/Source/Runtime/Engine/Classes/Particles/Trail/ParticleModuleTrailBase.h b/Engine/Source/Runtime/Engine/Classes/Particles/Trail/ParticleModuleTrailBase.h
--- a/Engine/Source/Runtime/Engine/Classes/Particles/Trail/ParticleModuleTrailBase.h
+++ b/Engine/Source/Runtime/Engine/Classes/Particles/Trail/ParticleModuleTrailBase.h
@@ -20,7 +20,7 @@ class UParticleModuleTrailBase : public UParticleModule
 	virtual EModuleType	GetModuleType() const override {	return EPMT_Trail;	}
 	virtual bool CanTickInAnyThread() override
 	{
-		return false;
+		return true;	
 	}
 	//~ End UParticleModule Interface
 };
diff --git a/Engine/Source/Runtime/Engine/Classes/PhysicsEngine/BodyInstance.h b/Engine/Source/Runtime/Engine/Classes/PhysicsEngine/BodyInstance.h
--- a/Engine/Source/Runtime/Engine/Classes/PhysicsEngine/BodyInstance.h
+++ b/Engine/Source/Runtime/Engine/Classes/PhysicsEngine/BodyInstance.h
@@ -1026,7 +1026,7 @@ private:
 	/** 
 	 * Helper function to update per shape filtering info. This should interface is not very friendly and should only be used from inside FBodyInstance
 	 */
-	void UpdatePhysicsShapeFilterData(uint32 SkelMeshCompID, bool bUseComplexAsSimple, bool bUseSimpleAsComplex, bool bPhysicsStatic, const TEnumAsByte<ECollisionEnabled::Type> * CollisionEnabledOverride, FCollisionResponseContainer * ResponseOverride, bool * bNotifyOverride);
+	void UpdatePhysicsShapeFilterData(uint32 ComponentID, bool bUseComplexAsSimple, bool bUseSimpleAsComplex, bool bPhysicsStatic, const TEnumAsByte<ECollisionEnabled::Type> * CollisionEnabledOverride, FCollisionResponseContainer * ResponseOverride, bool * bNotifyOverride);
 
 	/** Check if the shape is owned by this body instance */
 	bool IsShapeBoundToBody(const physx::PxShape* PShape) const;
diff --git a/Engine/Source/Runtime/Engine/Classes/Sound/SoundAttenuation.h b/Engine/Source/Runtime/Engine/Classes/Sound/SoundAttenuation.h
--- a/Engine/Source/Runtime/Engine/Classes/Sound/SoundAttenuation.h
+++ b/Engine/Source/Runtime/Engine/Classes/Sound/SoundAttenuation.h
@@ -72,7 +72,7 @@ struct ENGINE_API FAttenuationSettings
 	uint32 bEnableListenerFocus:1;
 
 	/** Whether or not to enable line-of-sight occlusion checking for the sound that plays in this audio component. */
-	UPROPERTY(EditAnywhere, BlueprintReadWrite, Category = Occlusion)
+	UPROPERTY(EditAnywhere, BlueprintReadWrite, Category = Occlusion, meta = (EditCondition = "bSpatialize"))
 	uint32 bEnableOcclusion:1;
 
 	/** Whether or not to enable complex geometry occlusion checks. */
@@ -157,23 +157,27 @@ struct ENGINE_API FAttenuationSettings
 	float NonFocusAzimuth;
 
 	/** Amount to scale the distance calculation of sounds that are in-focus. Can be used to make in-focus sounds appear to be closer or further away than they actually are. */
-	UPROPERTY(EditAnywhere, BlueprintReadWrite, Category = Focus, meta = (EditCondition = "bEnableListenerFocus"))
+	UPROPERTY(EditAnywhere, BlueprintReadWrite, Category = Focus, meta = (ClampMin = "0.0", UIMin = "0.0", EditCondition = "bEnableListenerFocus"))
 	float FocusDistanceScale;
 
 	/** Amount to scale the distance calculation of sounds that are not in-focus. Can be used to make in-focus sounds appear to be closer or further away than they actually are.  */
-	UPROPERTY(EditAnywhere, BlueprintReadWrite, Category = Focus, meta = (EditCondition = "bEnableListenerFocus"))
+	UPROPERTY(EditAnywhere, BlueprintReadWrite, Category = Focus, meta = (ClampMin = "0.0", UIMin = "0.0", EditCondition = "bEnableListenerFocus"))
 	float NonFocusDistanceScale;
 
 	/** Amount to scale the priority of sounds that are in focus. Can be used to boost the priority of sounds that are in focus. */
-	UPROPERTY(EditAnywhere, BlueprintReadWrite, Category = Focus, meta = (EditCondition = "bEnableListenerFocus"))
+	UPROPERTY(EditAnywhere, BlueprintReadWrite, Category = Focus, meta = (ClampMin = "0.0", UIMin = "0.0", EditCondition = "bEnableListenerFocus"))
 	float FocusPriorityScale;
 
 	/** Amount to scale the priority of sounds that are not in-focus. Can be used to reduce the priority of sounds that are not in focus. */
-	UPROPERTY(EditAnywhere, BlueprintReadWrite, Category = Focus, meta = (EditCondition = "bEnableListenerFocus"))
+	UPROPERTY(EditAnywhere, BlueprintReadWrite, Category = Focus, meta = (ClampMin = "0.0", UIMin = "0.0", EditCondition = "bEnableListenerFocus"))
 	float NonFocusPriorityScale;
 
+	/** Amount to attenuate sounds that are in focus. Can be overridden at the sound-level. */
+	UPROPERTY(EditAnywhere, BlueprintReadWrite, Category = Focus, meta = (ClampMin = "0.0", UIMin = "0.0", EditCondition = "bEnableListenerFocus"))
+	float FocusVolumeAttenuation;
+
 	/** Amount to attenuate sounds that are not in focus. Can be overridden at the sound-level. */
-	UPROPERTY(EditAnywhere, BlueprintReadWrite, Category = Focus, meta = (EditCondition = "bEnableListenerFocus"))
+	UPROPERTY(EditAnywhere, BlueprintReadWrite, Category = Focus, meta = (ClampMin = "0.0", UIMin = "0.0", EditCondition = "bEnableListenerFocus"))
 	float NonFocusVolumeAttenuation;
 
 	/** The low pass filter frequency (in hertz) to apply if the sound playing in this audio component is occluded. This will override the frequency set in LowPassFilterFrequency. A frequency of 0.0 is the device sample rate and will bypass the filter. */
@@ -217,6 +221,7 @@ struct ENGINE_API FAttenuationSettings
 		, NonFocusDistanceScale(1.0f)
 		, FocusPriorityScale(1.0f)
 		, NonFocusPriorityScale(1.0f)
+		, FocusVolumeAttenuation(1.0f)
 		, NonFocusVolumeAttenuation(1.0f)
 		, OcclusionLowPassFilterFrequency(20000.f)
 		, OcclusionVolumeAttenuation(1.0f)
@@ -236,6 +241,9 @@ struct ENGINE_API FAttenuationSettings
 
 	void CollectAttenuationShapesForVisualization(TMultiMap<EAttenuationShape::Type, AttenuationShapeDetails>& ShapeDetailsMap) const;
 	float GetMaxDimension() const;
+	float GetFocusPriorityScale(const struct FGlobalFocusSettings& FocusSettings, float FocusFactor) const;
+	float GetFocusAttenuation(const struct FGlobalFocusSettings& FocusSettings, float FocusFactor) const;
+	float GetFocusDistanceScale(const struct FGlobalFocusSettings& FocusSettings, float FocusFactor) const;
 	float AttenuationEval(const float Distance, const float Falloff, const float DistanceScale) const;
 	float AttenuationEvalBox(const FTransform& SoundLocation, const FVector ListenerLocation, const float DistanceScale) const;
 	float AttenuationEvalCapsule(const FTransform& SoundLocation, const FVector ListenerLocation, const float DistanceScale) const;
diff --git a/Engine/Source/Runtime/Engine/Classes/Sound/SoundBase.h b/Engine/Source/Runtime/Engine/Classes/Sound/SoundBase.h
--- a/Engine/Source/Runtime/Engine/Classes/Sound/SoundBase.h
+++ b/Engine/Source/Runtime/Engine/Classes/Sound/SoundBase.h
@@ -85,19 +85,6 @@ public:
 	virtual const FAttenuationSettings* GetAttenuationSettingsToApply() const;
 
 	/**
-	 * Checks to see if a location is audible
-	 */
-	bool IsAudible( const FVector& SourceLocation, const FVector& ListenerLocation, AActor* SourceActor, bool& bIsOccluded, bool bCheckOcclusion );
-
-	/** 
-	 * Does a simple range check to all listeners to test hearability
-	 *
-	 * @param Location				Location to check against
-	 * @param AttenuationSettings	Optional Attenuation override if not using settings from the sound
-	 */
-	bool IsAudibleSimple( class FAudioDevice* AudioDevice, const FVector Location, USoundAttenuation* InAttenuationSettings = NULL );
-
-	/** 
 	 * Returns the farthest distance at which the sound could be heard
 	 */
 	virtual float GetMaxAudibleDistance();
@@ -111,7 +98,7 @@ public:
 	virtual float GetPitchMultiplier();
 
 	/** Returns whether or not this sound is looping. */
-	const bool IsLooping() { return (GetDuration() >= INDEFINITELY_LOOPING_DURATION); }
+	bool IsLooping();
 
 	/** 
 	 * Parses the Sound to generate the WaveInstances to play
diff --git a/Engine/Source/Runtime/Engine/Classes/Sound/SoundConcurrency.h b/Engine/Source/Runtime/Engine/Classes/Sound/SoundConcurrency.h
--- a/Engine/Source/Runtime/Engine/Classes/Sound/SoundConcurrency.h
+++ b/Engine/Source/Runtime/Engine/Classes/Sound/SoundConcurrency.h
@@ -37,7 +37,7 @@ struct ENGINE_API FSoundConcurrencySettings
 	GENERATED_USTRUCT_BODY()
 
 	/** The max number of allowable concurrent active voices for voices playing in this concurrency group. */
-	UPROPERTY(EditAnywhere, BlueprintReadWrite, Category = Concurrency, meta = (ClampMin = "1"))
+	UPROPERTY(EditAnywhere, BlueprintReadWrite, Category = Concurrency, meta = (UIMin = "1", ClampMin = "1"))
 	int32 MaxCount;
 
 	/* Whether or not to limit the concurrency to per sound owner (i.e. the actor that plays the sound). If the sound doesn't have an owner, it falls back to global concurrency. */
@@ -48,9 +48,20 @@ struct ENGINE_API FSoundConcurrencySettings
 	UPROPERTY(EditAnywhere, BlueprintReadWrite, Category = Concurrency)
 	TEnumAsByte<enum EMaxConcurrentResolutionRule::Type> ResolutionRule;
 
-	/** The amount of attenuation to apply to each voice instance in this concurrency group. */
-	UPROPERTY(EditAnywhere, BlueprintReadWrite, Category = Concurrency, meta = (ClampMin = "0", ClampMax = "1"))
+	/** 
+	 * The amount of attenuation to apply to older voice instances in this concurrency group. This reduces volume of older voices in a concurrency group as new voices play.
+	 *
+	 * AppliedVolumeScale = Math.Pow(DuckingScale, VoiceGeneration)
+	 */
+	UPROPERTY(EditAnywhere, BlueprintReadWrite, Category = Concurrency, meta = (UIMin = "0", UIMax = "1", ClampMin = "0", ClampMax = "1"))
 	float VolumeScale;
+
+	FSoundConcurrencySettings()
+		: MaxCount(16)
+		, bLimitToOwner(false)
+		, ResolutionRule(EMaxConcurrentResolutionRule::StopFarthestThenOldest)
+		, VolumeScale(1.0f)
+	{}
 };
 
 UCLASS(BlueprintType, hidecategories=Object, editinlinenew, MinimalAPI)
@@ -83,6 +94,7 @@ class FConcurrencyGroup
 	int32 MaxActiveSounds;
 	FConcurrencyGroupID ConcurrencyGroupID;
 	EMaxConcurrentResolutionRule::Type ResolutionRule;
+	int32 Generation;
 
 public:
 	/** Constructor for the max concurrency active sound entry. */
@@ -96,6 +108,12 @@ public:
 	/** Retrieves the active sounds array. */
 	TArray<FActiveSound*>& GetActiveSounds();
 
+	/** Returns the number of active sounds in concurrency group. */
+	const int32 GetNumActiveSounds() const { return ActiveSounds.Num(); }
+
+	/** Retrieves the current generation */
+	const int32 GetGeneration() const { return Generation; }
+
 	/** Adds an active sound to the active sound array. */
 	void AddActiveSound(struct FActiveSound* ActiveSound);
 
diff --git a/Engine/Source/Runtime/Engine/Private/AI/Navigation/PImplRecastNavMesh.cpp b/Engine/Source/Runtime/Engine/Private/AI/Navigation/PImplRecastNavMesh.cpp
--- a/Engine/Source/Runtime/Engine/Private/AI/Navigation/PImplRecastNavMesh.cpp
+++ b/Engine/Source/Runtime/Engine/Private/AI/Navigation/PImplRecastNavMesh.cpp
@@ -855,7 +855,11 @@ ENavigationQueryResult::Type FPImplRecastNavMesh::FindPath(const FVector& StartL
 	if (dtStatusDetail(FindPathStatus, DT_INVALID_CYCLE_PATH))
 	{
 		UE_VLOG(NavMeshOwner, LogNavigation, Error, TEXT("FPImplRecastNavMesh::FindPath resulted in a cyclic path!"));
-		Path.DescribeSelfToVisLog(FVisualLogger::Get().GetLastEntryForObject(NavMeshOwner));
+		FVisualLogEntry* Entry = FVisualLogger::Get().GetLastEntryForObject(NavMeshOwner);
+		if (Entry)
+		{
+			Path.DescribeSelfToVisLog(Entry);
+		}
 	}
 #endif // ENABLE_VISUAL_LOG
 
