Copying //UE4/Dev-Framework to Dev-Main (//UE4/Dev-Main) @ 2879625

==========================
MAJOR FEATURES + CHANGES
==========================
Change 2821607 on 2016/01/08 by [REDACTED_PERSON]

	Added a way to limit amount of information logged by vlog by discarding logs from classes from outside of class whitelist #UE4

	This feature was followed by refactoring of functions taking FVisualLogEntry pointers to use references instead.

Change 2828384 on 2016/01/14 by [REDACTED_PERSON]

	Back out of visual log refactor done as part of CL#2821607 #UE4

Change 2869215 on 2016/02/16 by [REDACTED_PERSON]

	Store a WorldSettings pointer on ULevel instead of requiring it be index 0 in the Actors array. However, we will still generally attempt to keep it at index 0 for consistency with previous behavior

Change 2869404 on 2016/02/16 by [REDACTED_PERSON]

	Improve UI for default collision. It now uses a single drop down and sets the appropriate flags under the hood.

Change 2870062 on 2016/02/17 by Jurre.deBaare

	Name parameter driven by bone controller

Change 2870280 on 2016/02/17 by [REDACTED_PERSON]

	Vis log category handling fixes #UE4

	Also, a minor cleanup

Change 2871729 on 2016/02/18 by [REDACTED_PERSON]

	UE-26663 Fix 'LOD For Collision' display name

Change 2871730 on 2016/02/18 by [REDACTED_PERSON]

	UE-26580 Make ECollisionEnabled a BlueprintType
	UE-25373 Add a MakeHitResult node

Change 2871732 on 2016/02/18 by [REDACTED_PERSON]

	UE-24397 Add 'test' option to async query API, and use it in places that made sense. Also removed deprecated (4.8) functions from API.

Change 2872022 on 2016/02/18 by [REDACTED_PERSON]

	gameplay debugger refactor
	#ue4

Change 2872082 on 2016/02/18 by [REDACTED_PERSON]

	enabled old gameplay debugger as default one for now
	it will be deprecated with next version after testing in game projects
	#ue4

Change 2872390 on 2016/02/18 by [REDACTED_PERSON]

	OR-15041 (CPU) Hitches due to audio decompression on Windows

	1) Moving ogg-vorbis file info parsing into a worker thread
	 - stat dumphitches now shows the vorbis stuff totally gone

	2) Moving async decoding tasks to be retrieved and started from OnBufferEnd callback

Change 2872418 on 2016/02/18 by [REDACTED_PERSON]

	Fixed EQS debugger not storing data properly when subsequent Option is the one that produces result #UE4

Change 2872446 on 2016/02/18 by [REDACTED_PERSON]

	Using cached value of ActualVolume in GetVolumeWeightedPriority

Change 2872770 on 2016/02/18 by [REDACTED_PERSON]

	QAGame testing content for audio testing.

	Going to create a folder with specific sub-system testing maps for audio

Change 2873733 on 2016/02/19 by Jurre.deBaare

	- HLOD generated assets are now saved into a separate package instead of inside of the level asset

Change 2873828 on 2016/02/19 by [REDACTED_PERSON]

	Distributions that bake out no longer load in cooked build.

Change 2874623 on 2016/02/19 by [REDACTED_PERSON]

	UE-27131 Support for changing sound class volumes dynamically

	- new BP function to override a sound mix sound class adjuster
	- cleanup of AudioDevice.h and AudioDevice.cpp
	   - removing unnecessarily forward declares on various types
	   - removing unnecessary spaces and (void) params, etc

Change 2874922 on 2016/02/20 by [REDACTED_PERSON]

	Fixed EQS tests being compiled out from Shipping and Test with WITH_DEV_AUTOMATION_TESTS macro #UE4

Change 2875838 on 2016/02/22 by [REDACTED_PERSON]

diff --git a/Engine/Source/Developer/GameplayDebugger/Classes/GameplayDebuggerSettings.h b/Engine/Source/Developer/GameplayDebugger/Classes/GameplayDebuggerSettings.h
--- a/Engine/Source/Developer/GameplayDebugger/Classes/GameplayDebuggerSettings.h
+++ b/Engine/Source/Developer/GameplayDebugger/Classes/GameplayDebuggerSettings.h
@@ -1,11 +1,13 @@
 // Copyright 1998-2016 Epic Games, Inc. All Rights Reserved.
 
-/*=============================================================================
-GameplayDebuggerSettings.h: Declares the UGameplayDebuggerSettings class.
-=============================================================================*/
+//////////////////////////////////////////////////////////////////////////
+// THIS CLASS IS NOW DEPRECATED AND WILL BE REMOVED IN NEXT VERSION
+// Please check GameplayDebugger.h for details.
+
 #pragma once
 
 
+#include "GameplayDebuggingTypes.h"
 #include "GameplayDebuggerSettings.generated.h"
 
 #define ADD_LEVEL_EDITOR_EXTENSIONS 0
diff --git a/Engine/Source/Developer/GameplayDebugger/Classes/GameplayDebuggingComponent.h b/Engine/Source/Developer/GameplayDebugger/Classes/GameplayDebuggingComponent.h
--- a/Engine/Source/Developer/GameplayDebugger/Classes/GameplayDebuggingComponent.h
+++ b/Engine/Source/Developer/GameplayDebugger/Classes/GameplayDebuggingComponent.h
@@ -1,8 +1,8 @@
 // Copyright 1998-2016 Epic Games, Inc. All Rights Reserved.
 
-/**
- * GameplayDebuggingComponent is used to replicate debug data from server to client(s).
- */
+//////////////////////////////////////////////////////////////////////////
+// THIS CLASS IS NOW DEPRECATED AND WILL BE REMOVED IN NEXT VERSION
+// Please check GameplayDebugger.h for details.
 
 #pragma once
 #include "Components/PrimitiveComponent.h"
@@ -10,7 +10,6 @@
 #include "EnvironmentQuery/EQSQueryResultSourceInterface.h"
 #include "EnvironmentQuery/EnvQueryDebugHelpers.h"
 #include "Debug/DebugDrawService.h"
-#include "Debug/GameplayDebuggerBaseObject.h"
 #include "GameplayDebuggingComponent.generated.h"
 
 #define WITH_EQS 1
@@ -35,6 +34,60 @@ namespace EDebugComponentMessage
 
 DECLARE_MULTICAST_DELEGATE_TwoParams(FOnDebuggingTargetChanged, class AActor* /*Owner of debugging component*/, bool /*is being debugged now*/);
 
+DECLARE_LOG_CATEGORY_EXTERN(LogGameplayDebugger, Log, All);
+
+UENUM()
+enum class EGameplayDebuggerShapeElement : uint8
+{
+	Invalid,
+	String,
+	SinglePoint,
+	Segment,
+	Box,
+	Cone,
+	Cylinder,
+	Capsule,
+	Polygon,
+};
+
+USTRUCT()
+struct FGameplayDebuggerShapeElement
+{
+	GENERATED_USTRUCT_BODY()
+
+	/** points defining shape */
+	UPROPERTY()
+	TArray<FVector> Points;
+
+	UPROPERTY()
+	float ThicknesOrRadius;
+
+	/** description of shape */
+	UPROPERTY()
+	FString Description;
+
+	/** color of shape */
+	UPROPERTY()
+	FColor Color;
+
+	/** type of shape */
+	UPROPERTY()
+	EGameplayDebuggerShapeElement Type;
+
+	FGameplayDebuggerShapeElement() : Type(EGameplayDebuggerShapeElement::Invalid) {}
+
+	EGameplayDebuggerShapeElement GetType() const { return Type; }
+	FColor GetFColor() const { return Color; }
+
+	static FGameplayDebuggerShapeElement MakePoint(const FVector& Location, const float Radius, const FColor& Color, const FString& Description);
+	static FGameplayDebuggerShapeElement MakeSegment(const FVector& StartLocation, const FVector& EndLocation, const float Thickness, const FColor& Color, const FString& Description);
+	static FGameplayDebuggerShapeElement MakeBox(const FVector& Center, const FVector& Extent, const FColor& Color, const FString& Description);
+	static FGameplayDebuggerShapeElement MakeCone(const FVector& Location, const FVector& Direction, const float Length, const FColor& Color, const FString& Description);
+	static FGameplayDebuggerShapeElement MakeCylinder(const FVector& Center, const float Radius, const float HalfHeight, const FColor& Color, const FString& Description);
+	static FGameplayDebuggerShapeElement MakeCapsule(const FVector& Center, const float Radius, const float HalfHeight, const FColor& Color, const FString& Description);
+	static FGameplayDebuggerShapeElement MakePolygon(const TArray<FVector>& Verts, const FColor& Color, const FString& Description);
+};
+
 UCLASS(config=Engine)
 class GAMEPLAYDEBUGGER_API UGameplayDebuggingComponent : public UPrimitiveComponent, public IEQSQueryResultSourceInterface
 {
@@ -210,6 +263,10 @@ class GAMEPLAYDEBUGGER_API UGameplayDebuggingComponent : public UPrimitiveCompon
 	virtual void EnableDebugDraw(bool bEnable, bool InFocusedComponent = false);
 
 	bool ShouldReplicateData(EAIDebugDrawDataView::Type InView) const { return ReplicateViewDataCounters[InView] > 0 /*true*/; }
+	
+#if !ENABLE_OLD_GAMEPLAY_DEBUGGER
+	DEPRECATED_FORGAME(4.12, "GameplayDebuggingComponent class is now deprecated, please check GameplayDebugger.h for details.")
+#endif // !ENABLE_OLD_GAMEPLAY_DEBUGGER
 	virtual void CollectDataToReplicate(bool bCollectExtendedData);
 
 	//=============================================================================
diff --git a/Engine/Source/Developer/GameplayDebugger/Classes/GameplayDebuggingControllerComponent.h b/Engine/Source/Developer/GameplayDebugger/Classes/GameplayDebuggingControllerComponent.h
--- a/Engine/Source/Developer/GameplayDebugger/Classes/GameplayDebuggingControllerComponent.h
+++ b/Engine/Source/Developer/GameplayDebugger/Classes/GameplayDebuggingControllerComponent.h
@@ -1,8 +1,8 @@
 // Copyright 1998-2016 Epic Games, Inc. All Rights Reserved.
 
-/**
- * GameplayDebuggingComponent is used to replicate debug data from server to client(s).
- */
+//////////////////////////////////////////////////////////////////////////
+// THIS CLASS IS NOW DEPRECATED AND WILL BE REMOVED IN NEXT VERSION
+// Please check GameplayDebugger.h for details.
 
 #pragma once
 #include "TimerManager.h"
diff --git a/Engine/Source/Developer/GameplayDebugger/Classes/GameplayDebuggingHUDComponent.h b/Engine/Source/Developer/GameplayDebugger/Classes/GameplayDebuggingHUDComponent.h
--- a/Engine/Source/Developer/GameplayDebugger/Classes/GameplayDebuggingHUDComponent.h
+++ b/Engine/Source/Developer/GameplayDebugger/Classes/GameplayDebuggingHUDComponent.h
@@ -1,5 +1,9 @@
 // Copyright 1998-2016 Epic Games, Inc. All Rights Reserved.
 
+//////////////////////////////////////////////////////////////////////////
+// THIS CLASS IS NOW DEPRECATED AND WILL BE REMOVED IN NEXT VERSION
+// Please check GameplayDebugger.h for details.
+
 #pragma once
 #include "ShowFlags.h"
 #include "GameplayDebuggerSettings.h"
@@ -58,9 +62,13 @@ protected:
 	virtual void DrawBehaviorTreeData(APlayerController* PC, class UGameplayDebuggingComponent *DebugComponent);
 	virtual void DrawEQSData(APlayerController* PC, class UGameplayDebuggingComponent *DebugComponent);
 	virtual void DrawPerception(APlayerController* PC, class UGameplayDebuggingComponent *DebugComponent);
-	virtual void DrawGameSpecificView(APlayerController* PC, class UGameplayDebuggingComponent *DebugComponent) {}
 	virtual void DrawNavMeshSnapshot(APlayerController* PC, class UGameplayDebuggingComponent *DebugComponent);
 
+#if !ENABLE_OLD_GAMEPLAY_DEBUGGER
+	DEPRECATED_FORGAME(4.12, "GameplayDebuggingHUDComponent class is now deprecated, please check GameplayDebugger.h for details.")
+#endif // !ENABLE_OLD_GAMEPLAY_DEBUGGER
+	virtual void DrawGameSpecificView(APlayerController* PC, class UGameplayDebuggingComponent *DebugComponent) {}
+
 	void PrintAllData();
 	void DrawMenu(const float X, const float Y, class UGameplayDebuggingComponent* DebugComponent);
 	static void CalulateStringSize(const AGameplayDebuggingHUDComponent::FPrintContext& DefaultContext, UFont* Font, const FString& InString, float& OutX, float& OutY);
diff --git a/Engine/Source/Developer/GameplayDebugger/Classes/GameplayDebuggingReplicator.h b/Engine/Source/Developer/GameplayDebugger/Classes/GameplayDebuggingReplicator.h
--- a/Engine/Source/Developer/GameplayDebugger/Classes/GameplayDebuggingReplicator.h
+++ b/Engine/Source/Developer/GameplayDebugger/Classes/GameplayDebuggingReplicator.h
@@ -1,5 +1,9 @@
 // Copyright 1998-2016 Epic Games, Inc. All Rights Reserved.
 
+//////////////////////////////////////////////////////////////////////////
+// THIS CLASS IS NOW DEPRECATED AND WILL BE REMOVED IN NEXT VERSION
+// Please check GameplayDebugger.h for details.
+
 #pragma once
 #include "GameplayDebuggingTypes.h"
 #include "GameplayDebuggingReplicator.generated.h"
diff --git a/Engine/Source/Developer/GameplayDebugger/Classes/GameplayDebuggingTypes.h b/Engine/Source/Developer/GameplayDebugger/Classes/GameplayDebuggingTypes.h
--- a/Engine/Source/Developer/GameplayDebugger/Classes/GameplayDebuggingTypes.h
+++ b/Engine/Source/Developer/GameplayDebugger/Classes/GameplayDebuggingTypes.h
@@ -1,5 +1,9 @@
 // Copyright 1998-2016 Epic Games, Inc. All Rights Reserved.
 
+//////////////////////////////////////////////////////////////////////////
+// THIS CLASS IS NOW DEPRECATED AND WILL BE REMOVED IN NEXT VERSION
+// Please check GameplayDebugger.h for details.
+
 #pragma once
 
 #include "GameplayDebuggingTypes.generated.h"
diff --git a/Engine/Source/Developer/GameplayDebugger/GameplayDebugger.Build.cs b/Engine/Source/Developer/GameplayDebugger/GameplayDebugger.Build.cs
--- a/Engine/Source/Developer/GameplayDebugger/GameplayDebugger.Build.cs
+++ b/Engine/Source/Developer/GameplayDebugger/GameplayDebugger.Build.cs
@@ -6,65 +6,57 @@ namespace UnrealBuildTool.Rules
     {
         public GameplayDebugger(TargetInfo Target)
         {
-            PublicIncludePaths.AddRange(
-                new string[] {
-				    "Developer/GameplayDebugger/Public",
-				    "Developer/AIModule/Public",
-                    "Developer/Settings/Public",
-                    "Runtime/GameplayAbilities/Public",
-				    // ... add public include paths required here ...
-			    }
-            );
+            // change this flag to if you need to use deprecated GameplayDebugger module
+            // however, it will be removed in next engine version!
+            const bool bEnableDeprecatedDebugger = true;
 
-            PrivateIncludePaths.AddRange(
-                new string[] {
-					"Developer/GameplayDebugger/Private",
-                    "Runtime/Engine/Private",
-                    "Runtime/AIModule/Private",
-					// ... add other private include paths required here ...
-				    }
-                );
-
-            PrivateDependencyModuleNames.AddRange(
-                new string[]
-				{
+            PublicDependencyModuleNames.AddRange(
+				new string[] {
 					"Core",
 					"CoreUObject",
-                    "InputCore",
-					"Engine",    
-                    "RenderCore",
-                    "RHI",
-                    "ShaderCore",
-                    "Settings",
-                    "AIModule",  // it have to be here for now. It'll be changed to remove any dependency to AIModule in future
+					"Engine",
+				});
 
-                    // @todo do we really need to include these two? Maybe a smart interface would help here?
-                    "GameplayTasks",
-                    "GameplayAbilities",    
-				}
-                );
+			PrivateDependencyModuleNames.AddRange(
+				new string[] {
+					"RenderCore",
+					"InputCore",
+					"Slate",
+				});
 
-            DynamicallyLoadedModuleNames.AddRange(
-                new string[]
-				    {
-					    // ... add any modules that your module loads dynamically here ...
-				    }
-                );
+			PrivateIncludePaths.Add("Developer/GameplayDebugger/Private");
 
             if (UEBuildConfiguration.bBuildEditor == true)
-            {
-                PrivateDependencyModuleNames.Add("UnrealEd");
-                PrivateDependencyModuleNames.Add("LevelEditor");
-                PrivateDependencyModuleNames.Add("Slate");
-                PublicIncludePaths.Add("Editor/LevelEditor/Public");
-            }
-
-            if (UEBuildConfiguration.bCompileRecast)
-            {
-                PrivateDependencyModuleNames.Add("Navmesh");
-            }
+			{
+				PrivateDependencyModuleNames.Add("UnrealEd");
+			}
+
+            Definitions.Add("ENABLE_OLD_GAMEPLAY_DEBUGGER=" + (bEnableDeprecatedDebugger ? "1" : "0"));
+
+            // include for compatibility with old debugger
+            // to be removed with next version
+			{
+                PublicIncludePaths.AddRange(
+					new string[] {
+						"Developer/GameplayDebugger/Public",
+						"Developer/AIModule/Public",
+						"Developer/Settings/Public",
+						"Runtime/GameplayAbilities/Public",
+						// ... add public include paths required here ...
+					}
+                );
 
-			PrecompileForTargets = PrecompileTargetsType.Any;
+				PrivateIncludePaths.AddRange(
+					new string[] {
+						"Runtime/Engine/Private",
+						"Runtime/AIModule/Private",
+						// ... add other private include paths required here ...
+						}
+					);
+
+                PrivateDependencyModuleNames.AddRange(new string[] { "AIModule", "GameplayTasks", "GameplayAbilities" });
+                CircularlyReferencedDependentModules.AddRange(new string[] { "AIModule", "GameplayTasks", "GameplayAbilities" });
+			}
         }
     }
 }
diff --git a/Engine/Source/Developer/GameplayDebugger/Private/GameplayDebugger.cpp b/Engine/Source/Developer/GameplayDebugger/Private/GameplayDebugger.cpp
--- a/Engine/Source/Developer/GameplayDebugger/Private/GameplayDebugger.cpp
+++ b/Engine/Source/Developer/GameplayDebugger/Private/GameplayDebugger.cpp
@@ -1,6 +1,12 @@
 // Copyright 1998-2016 Epic Games, Inc. All Rights Reserved.
 
-#include "GameplayDebuggerPrivate.h"
+//////////////////////////////////////////////////////////////////////////
+// THIS CLASS IS NOW DEPRECATED AND WILL BE REMOVED IN NEXT VERSION
+// Please check GameplayDebugger.h for details.
+
+#if ENABLE_OLD_GAMEPLAY_DEBUGGER
+
+#include "GameplayDebuggerPrivatePCH.h"
 #include "Misc/CoreMisc.h"
 #include "GameplayDebuggingComponent.h"
 #include "GameplayDebuggingHUDComponent.h"
@@ -9,6 +15,7 @@
 #include "AISystem.h"
 #include "GameplayDebuggerSettings.h"
 #include "GameFramework/PlayerState.h"
+#include "EngineUtils.h"
 #if WITH_EDITOR
 #include "Editor/EditorEngine.h"
 #include "ISettingsModule.h"
@@ -22,7 +29,7 @@ FGameplayDebuggerSettings GameplayDebuggerSettings(class AGameplayDebuggingRepli
 	return FGameplayDebuggerSettings(Replicator == NULL ? Settings : Replicator->DebuggerShowFlags);
 }
 
-class FGameplayDebugger : public FSelfRegisteringExec, public IGameplayDebugger
+class FGameplayDebuggerCompat : public FSelfRegisteringExec, public GameplayDebugger
 {
 public:
 	// Begin IModuleInterface
@@ -53,6 +60,13 @@ private:
 	virtual bool CreateGameplayDebuggerForPlayerController(APlayerController* PlayerController) override;
 	virtual bool IsGameplayDebuggerActiveForPlayerController(APlayerController* PlayerController) override;
 
+	virtual void RegisterCategory(FName CategoryName, FOnGetCategory MakeInstanceDelegate, EGameplayDebuggerCategoryState CategoryState = EGameplayDebuggerCategoryState::Disabled, int32 SlotIdx = INDEX_NONE) {}
+	virtual void UnregisterCategory(FName CategoryName) {}
+	virtual void NotifyCategoriesChanged() {}
+	virtual void RegisterExtension(FName ExtensionName, IGameplayDebugger::FOnGetExtension MakeInstanceDelegate) {}
+	virtual void UnregisterExtension(FName ExtensionName) {}
+	virtual void NotifyExtensionsChanged() {}
+
 	bool DoesGameplayDebuggingReplicatorExistForPlayerController(APlayerController* PlayerController);
 
 	TMap<TWeakObjectPtr<UWorld>, TArray<TWeakObjectPtr<AGameplayDebuggingReplicator> > > AllReplicatorsPerWorlds;
@@ -62,10 +76,10 @@ private:
 #endif
 };
 
-IMPLEMENT_MODULE(FGameplayDebugger, GameplayDebugger)
+IMPLEMENT_MODULE(FGameplayDebuggerCompat, GameplayDebugger)
 
 // This code will execute after your module is loaded into memory (but after global variables are initialized, of course.)
-void FGameplayDebugger::StartupModule()
+void FGameplayDebuggerCompat::StartupModule()
 { 
 	//EMIT_CUSTOM_WARNING("/Engine/Source/Developer/GameplayDebugger module is deprecated and it's going to be removed with next UE4 version. Please use GameplayDebuggerPlugin instead.");
 	//UE_LOG(LogGameplayDebugger, Warning, TEXT("/Engine/Source/Developer/GameplayDebugger module is deprecated and it's going to be removed with next UE4 version. Please use GameplayDebuggerPlugin instead."));
@@ -73,11 +87,11 @@ void FGameplayDebugger::StartupModule()
 #if !(UE_BUILD_SHIPPING || UE_BUILD_TEST)
 	if (GEngine)
 	{
-		GEngine->OnWorldAdded().AddRaw(this, &FGameplayDebugger::WorldAdded);
-		GEngine->OnWorldDestroyed().AddRaw(this, &FGameplayDebugger::WorldDestroyed);
+		GEngine->OnWorldAdded().AddRaw(this, &FGameplayDebuggerCompat::WorldAdded);
+		GEngine->OnWorldDestroyed().AddRaw(this, &FGameplayDebuggerCompat::WorldDestroyed);
 #if WITH_EDITOR
-		GEngine->OnLevelActorAdded().AddRaw(this, &FGameplayDebugger::OnLevelActorAdded);
-		GEngine->OnLevelActorDeleted().AddRaw(this, &FGameplayDebugger::OnLevelActorDeleted);
+		GEngine->OnLevelActorAdded().AddRaw(this, &FGameplayDebuggerCompat::OnLevelActorAdded);
+		GEngine->OnLevelActorDeleted().AddRaw(this, &FGameplayDebuggerCompat::OnLevelActorDeleted);
 
 		ISettingsModule* SettingsModule = FModuleManager::GetModulePtr<ISettingsModule>("Settings");
 		if (SettingsModule != nullptr)
@@ -89,7 +103,7 @@ void FGameplayDebugger::StartupModule()
 				);
 		}
 
-		UGameplayDebuggerSettings::StaticClass()->GetDefaultObject<UGameplayDebuggerSettings>()->OnSettingChanged().AddRaw(this, &FGameplayDebugger::HandleSettingChanged);
+		UGameplayDebuggerSettings::StaticClass()->GetDefaultObject<UGameplayDebuggerSettings>()->OnSettingChanged().AddRaw(this, &FGameplayDebuggerCompat::HandleSettingChanged);
 
 #		if ADD_LEVEL_EDITOR_EXTENSIONS
 		const bool bExtendViewportMenu = UGameplayDebuggerSettings::StaticClass()->GetDefaultObject<UGameplayDebuggerSettings>()->bExtendViewportMenu;
@@ -97,7 +111,7 @@ void FGameplayDebugger::StartupModule()
 		{
 			// Register the extension with the level editor
 			{
-				ViewMenuExtender = FLevelEditorModule::FLevelEditorMenuExtender::CreateRaw(this, &FGameplayDebugger::OnExtendLevelEditorViewMenu);
+				ViewMenuExtender = FLevelEditorModule::FLevelEditorMenuExtender::CreateRaw(this, &FGameplayDebuggerCompat::OnExtendLevelEditorViewMenu);
 				FLevelEditorModule* LevelEditor = FModuleManager::LoadModulePtr<FLevelEditorModule>(TEXT("LevelEditor"));
 				if (LevelEditor)
 				{
@@ -112,7 +126,7 @@ void FGameplayDebugger::StartupModule()
 }
 
 #if WITH_EDITOR
-void FGameplayDebugger::HandleSettingChanged(FName PropertyName)
+void FGameplayDebuggerCompat::HandleSettingChanged(FName PropertyName)
 {
 #if ADD_LEVEL_EDITOR_EXTENSIONS
 	const bool bExtendViewportMenu = UGameplayDebuggerSettings::StaticClass()->GetDefaultObject<UGameplayDebuggerSettings>()->bExtendViewportMenu;
@@ -122,7 +136,7 @@ void FGameplayDebugger::HandleSettingChanged(FName PropertyName)
 		{
 			// Register the extension with the level editor
 			{
-				ViewMenuExtender = FLevelEditorModule::FLevelEditorMenuExtender::CreateRaw(this, &FGameplayDebugger::OnExtendLevelEditorViewMenu);
+				ViewMenuExtender = FLevelEditorModule::FLevelEditorMenuExtender::CreateRaw(this, &FGameplayDebuggerCompat::OnExtendLevelEditorViewMenu);
 				FLevelEditorModule* LevelEditor = FModuleManager::LoadModulePtr<FLevelEditorModule>(TEXT("LevelEditor"));
 				if (LevelEditor)
 				{
@@ -148,7 +162,7 @@ void FGameplayDebugger::HandleSettingChanged(FName PropertyName)
 
 // This function may be called during shutdown to clean up your module.  For modules that support dynamic reloading,
 // we call this function before unloading the module.
-void FGameplayDebugger::ShutdownModule()
+void FGameplayDebuggerCompat::ShutdownModule()
 {
 #if !(UE_BUILD_SHIPPING || UE_BUILD_TEST)
 	if (GEngine)
@@ -187,7 +201,7 @@ void FGameplayDebugger::ShutdownModule()
 
 #if WITH_EDITOR
 
-void FGameplayDebugger::CreateSettingSubMenu(FMenuBuilder& Builder)
+void FGameplayDebuggerCompat::CreateSettingSubMenu(FMenuBuilder& Builder)
 {
 #if !(UE_BUILD_SHIPPING || UE_BUILD_TEST)
 	Builder.AddMenuEntry(
@@ -205,7 +219,7 @@ void FGameplayDebugger::CreateSettingSubMenu(FMenuBuilder& Builder)
 #endif
 }
 
-void FGameplayDebugger::CreateSnappingOptionsMenu(FMenuBuilder& Builder)
+void FGameplayDebuggerCompat::CreateSnappingOptionsMenu(FMenuBuilder& Builder)
 {
 #if !(UE_BUILD_SHIPPING || UE_BUILD_TEST)
 	if (GCurrentLevelEditingViewportClient && GCurrentLevelEditingViewportClient->EngineShowFlags.DebugAI && GCurrentLevelEditingViewportClient->IsSimulateInEditorViewport())
@@ -214,13 +228,13 @@ void FGameplayDebugger::CreateSnappingOptionsMenu(FMenuBuilder& Builder)
 		Builder.AddSubMenu(
 			LOCTEXT("Test_GameplayDebugger_SnappingOptions_Menu", "Gameplay Debugger"),
 			LOCTEXT("Test_GameplayDebugger_SnappingOptions_Menu_Tooltip", "Quick setting for Gameplay Debugger tool in selected view"),
-			FNewMenuDelegate::CreateRaw(this, &FGameplayDebugger::CreateSettingSubMenu)
+			FNewMenuDelegate::CreateRaw(this, &FGameplayDebuggerCompat::CreateSettingSubMenu)
 			);
 	}
 #endif
 }
 
-TSharedRef<FExtender> FGameplayDebugger::OnExtendLevelEditorViewMenu(const TSharedRef<FUICommandList> CommandList)
+TSharedRef<FExtender> FGameplayDebuggerCompat::OnExtendLevelEditorViewMenu(const TSharedRef<FUICommandList> CommandList)
 {
 #if !(UE_BUILD_SHIPPING || UE_BUILD_TEST) && WITH_EDITOR
 	TSharedRef<FExtender> Extender(new FExtender());
@@ -231,7 +245,7 @@ TSharedRef<FExtender> FGameplayDebugger::OnExtendLevelEditorViewMenu(const TShar
 			"LevelViewportLayouts",
 			EExtensionHook::After,
 			NULL,
-			FMenuExtensionDelegate::CreateRaw(this, &FGameplayDebugger::CreateSnappingOptionsMenu));
+			FMenuExtensionDelegate::CreateRaw(this, &FGameplayDebuggerCompat::CreateSnappingOptionsMenu));
 
 	}
 #endif
@@ -239,7 +253,7 @@ TSharedRef<FExtender> FGameplayDebugger::OnExtendLevelEditorViewMenu(const TShar
 }
 #endif
 
-bool FGameplayDebugger::DoesGameplayDebuggingReplicatorExistForPlayerController(APlayerController* PlayerController)
+bool FGameplayDebuggerCompat::DoesGameplayDebuggingReplicatorExistForPlayerController(APlayerController* PlayerController)
 {
 #if !(UE_BUILD_SHIPPING || UE_BUILD_TEST)
 	if (PlayerController == NULL)
@@ -266,7 +280,7 @@ bool FGameplayDebugger::DoesGameplayDebuggingReplicatorExistForPlayerController(
 	return false;
 }
 
-bool FGameplayDebugger::CreateGameplayDebuggerForPlayerController(APlayerController* PlayerController)
+bool FGameplayDebuggerCompat::CreateGameplayDebuggerForPlayerController(APlayerController* PlayerController)
 {
 #if !(UE_BUILD_SHIPPING || UE_BUILD_TEST)
 	if (PlayerController == NULL)
@@ -318,7 +332,7 @@ bool FGameplayDebugger::CreateGameplayDebuggerForPlayerController(APlayerControl
 	return false;
 }
 
-bool FGameplayDebugger::IsGameplayDebuggerActiveForPlayerController(APlayerController* PlayerController)
+bool FGameplayDebuggerCompat::IsGameplayDebuggerActiveForPlayerController(APlayerController* PlayerController)
 {
 #if !(UE_BUILD_SHIPPING || UE_BUILD_TEST)
 	if (PlayerController == NULL)
@@ -345,22 +359,22 @@ bool FGameplayDebugger::IsGameplayDebuggerActiveForPlayerController(APlayerContr
 	return false;
 }
 
-TArray<TWeakObjectPtr<AGameplayDebuggingReplicator> >& FGameplayDebugger::GetAllReplicators(UWorld* InWorld)
+TArray<TWeakObjectPtr<AGameplayDebuggingReplicator> >& FGameplayDebuggerCompat::GetAllReplicators(UWorld* InWorld)
 {
 	return AllReplicatorsPerWorlds.FindOrAdd(InWorld);
 }
 
-void FGameplayDebugger::AddReplicator(UWorld* InWorld, AGameplayDebuggingReplicator* InReplicator)
+void FGameplayDebuggerCompat::AddReplicator(UWorld* InWorld, AGameplayDebuggingReplicator* InReplicator)
 {
 	GetAllReplicators(InWorld).Add(InReplicator);
 }
 
-void FGameplayDebugger::RemoveReplicator(UWorld* InWorld, AGameplayDebuggingReplicator* InReplicator)
+void FGameplayDebuggerCompat::RemoveReplicator(UWorld* InWorld, AGameplayDebuggingReplicator* InReplicator)
 {
 	GetAllReplicators(InWorld).RemoveSwap(InReplicator);
 }
 
-void FGameplayDebugger::WorldAdded(UWorld* InWorld)
+void FGameplayDebuggerCompat::WorldAdded(UWorld* InWorld)
 {
 #if !(UE_BUILD_SHIPPING || UE_BUILD_TEST)
 	bool bIsServer = InWorld && InWorld->GetNetMode() < ENetMode::NM_Client; // (Only server code)
@@ -401,7 +415,7 @@ void FGameplayDebugger::WorldAdded(UWorld* InWorld)
 #endif //!(UE_BUILD_SHIPPING || UE_BUILD_TEST)
 }
 
-void FGameplayDebugger::WorldDestroyed(UWorld* InWorld)
+void FGameplayDebuggerCompat::WorldDestroyed(UWorld* InWorld)
 {
 #if !(UE_BUILD_SHIPPING || UE_BUILD_TEST)
 	bool bIsServer = InWorld && InWorld->GetNetMode() < ENetMode::NM_Client; // (Only work on  server)
@@ -416,7 +430,7 @@ void FGameplayDebugger::WorldDestroyed(UWorld* InWorld)
 }
 
 #if WITH_EDITOR
-void FGameplayDebugger::OnLevelActorAdded(AActor* InActor)
+void FGameplayDebuggerCompat::OnLevelActorAdded(AActor* InActor)
 {
 	// This function doesn't help much, because it's only called in EDITOR!
 	// We need a function that is called in the game!  So instead of creating it automatically, I'm leaving it
@@ -434,7 +448,7 @@ void FGameplayDebugger::OnLevelActorAdded(AActor* InActor)
 #endif
 }
 
-void FGameplayDebugger::OnLevelActorDeleted(AActor* InActor)
+void FGameplayDebuggerCompat::OnLevelActorDeleted(AActor* InActor)
 {
 #if !(UE_BUILD_SHIPPING || UE_BUILD_TEST)
 	if (!InActor)
@@ -475,7 +489,7 @@ void FGameplayDebugger::OnLevelActorDeleted(AActor* InActor)
 }
 #endif //WITH_EDITOR
 
-bool FGameplayDebugger::Exec(UWorld* InWorld, const TCHAR* Cmd, FOutputDevice& Ar)
+bool FGameplayDebuggerCompat::Exec(UWorld* InWorld, const TCHAR* Cmd, FOutputDevice& Ar)
 {
 #if !(UE_BUILD_SHIPPING || UE_BUILD_TEST)
 	bool bHandled = false;
@@ -685,3 +699,17 @@ bool FGameplayDebugger::Exec(UWorld* InWorld, const TCHAR* Cmd, FOutputDevice& A
 }
 
 #undef LOCTEXT_NAMESPACE
+
+#else // ENABLE_OLD_GAMEPLAY_DEBUGGER
+
+#include "GameplayDebuggerPrivatePCH.h"
+#include "GameplayDebuggingReplicator.h"
+#include "GameplayDebuggerSettings.h"
+
+FGameplayDebuggerSettings GameplayDebuggerSettings(class AGameplayDebuggingReplicator* Replicator)
+{
+	uint32& Settings = UGameplayDebuggerSettings::StaticClass()->GetDefaultObject<UGameplayDebuggerSettings>()->GetSettings();
+	return FGameplayDebuggerSettings(Replicator == NULL ? Settings : Replicator->DebuggerShowFlags);
+}
+
+#endif // ENABLE_OLD_GAMEPLAY_DEBUGGER
diff --git a/Engine/Source/Developer/GameplayDebugger/Private/GameplayDebuggerCategoryReplicator.cpp b/Engine/Source/Developer/GameplayDebugger/Private/GameplayDebuggerCategoryReplicator.cpp
new file mode 100644
--- /dev/null
+++ b/Engine/Source/Developer/GameplayDebugger/Private/GameplayDebuggerCategoryReplicator.cpp
@@ -0,0 +1,644 @@
+// Copyright 1998-2016 Epic Games, Inc. All Rights Reserved.
+
+#include "GameplayDebuggerPrivatePCH.h"
+#include "GameplayDebuggerCategoryReplicator.h"
+#include "GameplayDebuggerAddonManager.h"
+#include "GameplayDebuggerPlayerManager.h"
+#include "GameplayDebuggerRenderingComponent.h"
+#include "Net/UnrealNetwork.h"
+
+//////////////////////////////////////////////////////////////////////////
+// FGameplayDebuggerCategoryReplicatorData
+
+DEFINE_LOG_CATEGORY_STATIC(LogGameplayDebugReplication, Display, All);
+
+class FNetFastCategoryBaseState : public INetDeltaBaseState
+{
+public:
+	struct FCategoryState
+	{
+		int32 BaseRepCounter;
+		int32 TextLinesRepCounter;
+		int32 ShapesRepCounter;
+		TArray<int32> DataPackRepCounters;
+
+		bool operator==(const FCategoryState& OtherState) const
+		{
+			return (BaseRepCounter == OtherState.BaseRepCounter) &&
+				(TextLinesRepCounter == OtherState.TextLinesRepCounter) &&
+				(ShapesRepCounter == OtherState.ShapesRepCounter) &&
+				(DataPackRepCounters == OtherState.DataPackRepCounters);
+		}
+	};
+
+	virtual bool IsStateEqual(INetDeltaBaseState* OtherState) override
+	{
+		FNetFastCategoryBaseState* Other = static_cast<FNetFastCategoryBaseState*>(OtherState);
+		return (CategoryStates == Other->CategoryStates);
+	}
+
+	TArray<FCategoryState> CategoryStates;
+};
+
+bool FGameplayDebuggerNetPack::NetDeltaSerialize(FNetDeltaSerializeInfo& DeltaParms)
+{
+	if (DeltaParms.bUpdateUnmappedObjects || Owner == nullptr)
+	{
+		return true;
+	}
+
+	const uint32 DataPackMarker = 0x11223344;
+
+	if (DeltaParms.Writer)
+	{
+		FBitWriter& Writer = *DeltaParms.Writer;
+		int32 NumChangedCategories = 0;
+
+		FNetFastCategoryBaseState* OldState = static_cast<FNetFastCategoryBaseState*>(DeltaParms.OldState);
+		FNetFastCategoryBaseState* NewState = nullptr;
+		TArray<uint8> ChangedCategories;
+
+		// find delta to replicate
+		if (Owner->bIsEnabled && Owner->Categories.Num() == SavedData.Num())
+		{
+			NewState = new FNetFastCategoryBaseState();
+			check(DeltaParms.NewState);
+			*DeltaParms.NewState = TSharedPtr<INetDeltaBaseState>(NewState);
+			NewState->CategoryStates.SetNum(SavedData.Num());
+			ChangedCategories.AddZeroed(SavedData.Num());
+
+			for (int32 Idx = 0; Idx < SavedData.Num(); Idx++)
+			{
+				FGameplayDebuggerCategory& CategoryOb = Owner->Categories[Idx].Get();
+				FCategoryData& SavedCategory = SavedData[Idx];
+				const bool bMissingOldState = (OldState == nullptr) || !OldState->CategoryStates.IsValidIndex(Idx);
+				ChangedCategories[Idx] = bMissingOldState ? 1 : 0;
+
+				if (bMissingOldState || (SavedCategory.bIsEnabled != CategoryOb.bIsEnabled))
+				{
+					SavedCategory.bIsEnabled = CategoryOb.bIsEnabled;
+					NewState->CategoryStates[Idx].BaseRepCounter = bMissingOldState ? 1 : (OldState->CategoryStates[Idx].BaseRepCounter + 1);
+					ChangedCategories[Idx]++;
+				}
+
+				if (bMissingOldState || (SavedCategory.TextLines != CategoryOb.ReplicatedLines))
+				{
+					SavedCategory.TextLines = CategoryOb.ReplicatedLines;
+					NewState->CategoryStates[Idx].TextLinesRepCounter = bMissingOldState ? 1 : (OldState->CategoryStates[Idx].TextLinesRepCounter + 1);
+					ChangedCategories[Idx]++;
+				}
+
+				if (bMissingOldState || (SavedCategory.Shapes != CategoryOb.ReplicatedShapes))
+				{
+					SavedCategory.Shapes = CategoryOb.ReplicatedShapes;
+					NewState->CategoryStates[Idx].ShapesRepCounter = bMissingOldState ? 1 : (OldState->CategoryStates[Idx].ShapesRepCounter + 1);
+					ChangedCategories[Idx]++;
+				}
+
+				const int32 NumDataPacks = CategoryOb.ReplicatedDataPacks.Num();
+				SavedCategory.DataPacks.SetNum(NumDataPacks);
+				NewState->CategoryStates[Idx].DataPackRepCounters.SetNum(NumDataPacks);
+				for (int32 DataIdx = 0; DataIdx < CategoryOb.ReplicatedDataPacks.Num(); DataIdx++)
+				{
+					if (bMissingOldState || !(SavedCategory.DataPacks[DataIdx] == CategoryOb.ReplicatedDataPacks[DataIdx].Header))
+					{
+						SavedCategory.DataPacks[DataIdx] = CategoryOb.ReplicatedDataPacks[DataIdx].Header;
+						NewState->CategoryStates[Idx].DataPackRepCounters[DataIdx] =
+							bMissingOldState || !OldState->CategoryStates[Idx].DataPackRepCounters.IsValidIndex(DataIdx) ? 1 : (OldState->CategoryStates[Idx].DataPackRepCounters[DataIdx] + 1);
+
+						ChangedCategories[Idx]++;
+					}
+				}
+
+				NumChangedCategories += ChangedCategories[Idx] ? 1 : 0;
+			}
+		}
+
+		if (NumChangedCategories == 0)
+		{
+			return false;
+		}
+
+		int32 CategoryCount = SavedData.Num();
+		Writer << CategoryCount;
+
+		UE_LOG(LogGameplayDebugReplication, Verbose, TEXT("NetDeltaSerialize WRITE START, OldState:%d"), OldState ? 1 : 0);
+		for (int32 Idx = 0; Idx < SavedData.Num(); Idx++)
+		{
+			FGameplayDebuggerCategory& CategoryOb = Owner->Categories[Idx].Get();
+			const bool bMissingOldState = (OldState == nullptr) || !OldState->CategoryStates.IsValidIndex(Idx);
+			FCategoryData& SavedCategory = SavedData[Idx];
+
+			uint8 BaseFlags = SavedCategory.bIsEnabled;
+			uint8 ShouldUpdateTextLines = bMissingOldState || (OldState->CategoryStates[Idx].TextLinesRepCounter != NewState->CategoryStates[Idx].TextLinesRepCounter);
+			uint8 ShouldUpdateShapes = bMissingOldState || (OldState->CategoryStates[Idx].ShapesRepCounter != NewState->CategoryStates[Idx].ShapesRepCounter);
+			uint8 NumDataPacks = SavedCategory.DataPacks.Num();
+
+			Writer.WriteBit(BaseFlags);
+			Writer.WriteBit(ShouldUpdateTextLines);
+			Writer.WriteBit(ShouldUpdateShapes);
+			Writer << NumDataPacks;
+
+			if (ChangedCategories[Idx])
+			{
+				UE_LOG(LogGameplayDebugReplication, Verbose, TEXT("  CATEGORY[%d]:%s"), Idx, *CategoryOb.GetCategoryName().ToString());
+			}
+
+			if (ShouldUpdateTextLines)
+			{
+				UE_LOG(LogGameplayDebugReplication, Verbose, TEXT("  >> replicate lines"));
+				Writer << SavedCategory.TextLines;
+			}
+
+			if (ShouldUpdateShapes)
+			{
+				UE_LOG(LogGameplayDebugReplication, Verbose, TEXT("  >> replicate shapes"));
+				Writer << SavedCategory.Shapes;
+			}
+
+			for (int32 DataIdx = 0; DataIdx < NumDataPacks; DataIdx++)
+			{
+				uint8 ShouldUpdateDataPack = bMissingOldState || !OldState->CategoryStates[Idx].DataPackRepCounters.IsValidIndex(DataIdx) ||
+					(OldState->CategoryStates[Idx].DataPackRepCounters[DataIdx] != NewState->CategoryStates[Idx].DataPackRepCounters[DataIdx]);
+
+				Writer.WriteBit(ShouldUpdateDataPack);
+				if (ShouldUpdateDataPack)
+				{
+					// send single packet
+					FGameplayDebuggerDataPack& DataPack = CategoryOb.ReplicatedDataPacks[DataIdx];
+
+					uint8 IsCompressed = DataPack.Header.bIsCompressed ? 1 : 0;
+					Writer.WriteBit(IsCompressed);
+
+					Writer << DataPack.Header.DataVersion;
+					Writer << DataPack.Header.DataSize;
+					Writer << DataPack.Header.DataOffset;
+
+					const int32 PacketSize = FMath::Min(FGameplayDebuggerDataPack::PacketSize, DataPack.Header.DataSize - DataPack.Header.DataOffset);
+					if (PacketSize > 0)
+					{
+						Writer.Serialize(DataPack.Data.GetData() + DataPack.Header.DataOffset, PacketSize);
+					}
+
+					uint32 Marker = DataPackMarker;
+					Writer << Marker;
+
+					UE_LOG(LogGameplayDebugReplication, Verbose, TEXT("  >> replicate data pack[%d] progress:%0.f%% (offset:%d packet:%d)"), DataIdx, DataPack.GetProgress() * 100.0f, DataPack.Header.DataOffset, PacketSize);
+				}
+			}
+		}
+	}
+	else if (DeltaParms.Reader)
+	{
+		FBitReader& Reader = *DeltaParms.Reader;
+
+		UE_LOG(LogGameplayDebugReplication, Verbose, TEXT("NetDeltaSerialize READ START"));
+
+		int32 CategoryCount = 0;
+		Reader << CategoryCount;
+
+		if (CategoryCount != Owner->Categories.Num())
+		{
+			UE_LOG(LogGameplayDebugReplication, Error, TEXT("Category count mismtach! received:%d expected:%d"), CategoryCount, Owner->Categories.Num());
+			Reader.SetError();
+			return false;
+		}
+
+		for (int32 Idx = 0; Idx < CategoryCount; Idx++)
+		{
+			FGameplayDebuggerCategory& CategoryOb = Owner->Categories[Idx].Get();
+			UE_LOG(LogGameplayDebugReplication, Verbose, TEXT("  CATEGORY[%d]:%s"), Idx, *CategoryOb.GetCategoryName().ToString());
+
+			uint8 BaseFlags = Reader.ReadBit();
+			uint8 ShouldUpdateTextLines = Reader.ReadBit();
+			uint8 ShouldUpdateShapes = Reader.ReadBit();
+
+			uint8 NumDataPacks = 0;
+			Reader << NumDataPacks;
+
+			if ((int32)NumDataPacks != CategoryOb.ReplicatedDataPacks.Num())
+			{
+				UE_LOG(LogGameplayDebugReplication, Error, TEXT("Data pack count mismtach! received:%d expected:%d"), NumDataPacks, CategoryOb.ReplicatedDataPacks.Num());
+				Reader.SetError();
+				return false;
+			}
+
+			CategoryOb.bIsEnabled = (BaseFlags != 0);
+
+			if (ShouldUpdateTextLines)
+			{
+				UE_LOG(LogGameplayDebugReplication, Verbose, TEXT("  >> received lines"));
+				Reader << CategoryOb.ReplicatedLines;
+			}
+
+			if (ShouldUpdateShapes)
+			{
+				UE_LOG(LogGameplayDebugReplication, Verbose, TEXT("  >> received shapes"));
+				Reader << CategoryOb.ReplicatedShapes;
+			}
+
+			for (uint8 DataIdx = 0; DataIdx < NumDataPacks; DataIdx++)
+			{
+				uint8 ShouldUpdateDataPack = Reader.ReadBit();
+
+				if (ShouldUpdateDataPack)
+				{
+					// receive single packet
+					FGameplayDebuggerDataPack DataPacket;
+
+					uint8 IsCompressed = Reader.ReadBit();
+					DataPacket.Header.bIsCompressed = (IsCompressed != 0);
+
+					Reader << DataPacket.Header.DataVersion;
+					Reader << DataPacket.Header.DataSize;
+					Reader << DataPacket.Header.DataOffset;
+
+					const int32 PacketSize = FMath::Min(FGameplayDebuggerDataPack::PacketSize, DataPacket.Header.DataSize - DataPacket.Header.DataOffset);
+					if (PacketSize > 0)
+					{
+						DataPacket.Data.AddUninitialized(PacketSize);
+						Reader.Serialize(DataPacket.Data.GetData(), PacketSize);
+					}
+
+					uint32 Marker = 0;
+					Reader << Marker;
+
+					if (Marker != DataPackMarker)
+					{				
+						UE_LOG(LogGameplayDebugReplication, Error, TEXT("Data pack[%d] corrupted! marker:%X expected:%X"), DataIdx, Marker, DataPackMarker);
+						Reader.SetError();
+						return false;
+					}
+
+					Owner->OnReceivedDataPackPacket(Idx, DataIdx, DataPacket);
+					UE_LOG(LogGameplayDebugReplication, Verbose, TEXT("  >> replicate data pack[%d] progress:%.0f%%"), DataIdx, CategoryOb.ReplicatedDataPacks[DataIdx].GetProgress() * 100.0f);
+				}
+			}
+		}
+	}
+
+	return true;
+}
+
+void FGameplayDebuggerNetPack::OnCategoriesChanged()
+{
+	SavedData.Reset();
+	SavedData.SetNum(Owner->Categories.Num());
+}
+
+//////////////////////////////////////////////////////////////////////////
+// AGameplayDebuggerCategoryReplicator
+
+AGameplayDebuggerCategoryReplicator::AGameplayDebuggerCategoryReplicator(const FObjectInitializer& ObjectInitializer) : Super(ObjectInitializer)
+{
+	PrimaryActorTick.bCanEverTick = true;
+	PrimaryActorTick.bAllowTickOnDedicatedServer = true;
+	PrimaryActorTick.bTickEvenWhenPaused = true;
+	PrimaryActorTick.bStartWithTickEnabled = false;
+
+#if WITH_EDITOR
+	SetIsTemporarilyHiddenInEditor(true);
+#endif
+
+#if WITH_EDITORONLY_DATA
+	bHiddenEdLevel = true;
+	bHiddenEdLayer = true;
+	bHiddenEd = true;
+	bEditable = false;
+#endif
+
+	bIsEnabled = false;
+	bReplicates = true;
+
+	ReplicatedData.Owner = this;
+}
+
+void AGameplayDebuggerCategoryReplicator::BeginPlay()
+{
+	Super::BeginPlay();
+
+	UWorld* World = GetWorld();
+	const ENetMode NetMode = World->GetNetMode();
+	bHasAuthority = (NetMode != NM_Client);
+	bIsLocal = (NetMode != NM_DedicatedServer);
+
+	FGameplayDebuggerAddonManager& CategoryManager = FGameplayDebuggerAddonManager::GetCurrent();
+	CategoryManager.OnCategoriesChanged.AddUObject(this, &AGameplayDebuggerCategoryReplicator::OnCategoriesChanged);
+	
+	OnCategoriesChanged();
+
+	AGameplayDebuggerPlayerManager& PlayerManager = AGameplayDebuggerPlayerManager::GetCurrent(GetWorld());
+	PlayerManager.RegisterReplicator(*this);
+
+	SetActorHiddenInGame(!bIsLocal);
+	if (bIsLocal)
+	{
+		RenderingComp = NewObject<UGameplayDebuggerRenderingComponent>(this, TEXT("RenderingComp"));
+		RenderingComp->RegisterComponentWithWorld(World);
+		RootComponent = RenderingComp;
+	}
+	
+	if (bHasAuthority)
+	{
+		SetEnabled(FGameplayDebuggerAddonBase::IsSimulateInEditor());
+	}
+}
+
+void AGameplayDebuggerCategoryReplicator::Destroyed()
+{
+	Super::Destroyed();
+
+	FGameplayDebuggerAddonManager& CategoryManager = FGameplayDebuggerAddonManager::GetCurrent();
+	CategoryManager.OnCategoriesChanged.RemoveAll(this);
+}
+
+void AGameplayDebuggerCategoryReplicator::OnCategoriesChanged()
+{
+	FGameplayDebuggerAddonManager& CategoryManager = FGameplayDebuggerAddonManager::GetCurrent();
+	CategoryManager.CreateCategories(*this, Categories);
+
+	ReplicatedData.OnCategoriesChanged();
+
+	if (bIsLocal)
+	{
+		AGameplayDebuggerPlayerManager& PlayerManager = AGameplayDebuggerPlayerManager::GetCurrent(GetWorld());
+		PlayerManager.RefreshInputBindings(*this);
+	}
+}
+
+UNetConnection* AGameplayDebuggerCategoryReplicator::GetNetConnection() const
+{
+	return IsValid(OwnerPC) ? OwnerPC->GetNetConnection() : nullptr;
+}
+
+bool AGameplayDebuggerCategoryReplicator::IsNetRelevantFor(const AActor* RealViewer, const AActor* ViewTarget, const FVector& SrcLocation) const
+{
+	return (RealViewer == OwnerPC);
+}
+
+void AGameplayDebuggerCategoryReplicator::PostNetReceive()
+{
+	Super::PostNetReceive();
+
+	if (PendingReplicationRequests.Num() && !bHasAuthority)
+	{
+		ServerRequestDataPackets(PendingReplicationRequests);
+		PendingReplicationRequests.Reset();
+	}
+}
+
+void AGameplayDebuggerCategoryReplicator::GetLifetimeReplicatedProps(TArray<FLifetimeProperty>& OutLifetimeProps) const
+{
+	Super::GetLifetimeReplicatedProps(OutLifetimeProps);
+
+	DOREPLIFETIME(AGameplayDebuggerCategoryReplicator, OwnerPC);
+	DOREPLIFETIME(AGameplayDebuggerCategoryReplicator, DebugActor);
+	DOREPLIFETIME(AGameplayDebuggerCategoryReplicator, bIsEnabled);
+	DOREPLIFETIME(AGameplayDebuggerCategoryReplicator, ReplicatedData);
+}
+
+bool AGameplayDebuggerCategoryReplicator::ServerSetEnabled_Validate(bool bEnable)
+{
+	return true;
+}
+
+void AGameplayDebuggerCategoryReplicator::ServerSetEnabled_Implementation(bool bEnable)
+{
+	SetEnabled(bEnable);
+}
+
+bool AGameplayDebuggerCategoryReplicator::ServerSetDebugActor_Validate(AActor* Actor)
+{
+	return true;
+}
+
+void AGameplayDebuggerCategoryReplicator::ServerSetDebugActor_Implementation(AActor* Actor)
+{
+	SetDebugActor(Actor);
+}
+
+bool AGameplayDebuggerCategoryReplicator::ServerSetCategoryEnabled_Validate(int32 CategoryId, bool bEnable)
+{
+	return true;
+}
+
+void AGameplayDebuggerCategoryReplicator::ServerSetCategoryEnabled_Implementation(int32 CategoryId, bool bEnable)
+{
+	SetCategoryEnabled(CategoryId, bEnable);
+}
+
+bool AGameplayDebuggerCategoryReplicator::ServerRequestDataPackets_Validate(const TArray<FGameplayDebuggerDataPacketRequest>& RequestList)
+{
+	return true;
+}
+
+void AGameplayDebuggerCategoryReplicator::ServerRequestDataPackets_Implementation(const TArray<FGameplayDebuggerDataPacketRequest>& RequestList)
+{
+	for (int32 Idx = 0; Idx < RequestList.Num(); Idx++)
+	{
+		const FGameplayDebuggerDataPacketRequest& Request = RequestList[Idx];
+		if (Categories.IsValidIndex(Request.CategoryId) && Categories[Request.CategoryId]->ReplicatedDataPacks.IsValidIndex(Request.DataPackId))
+		{
+			Categories[Request.CategoryId]->ReplicatedDataPacks[Request.DataPackId].OnPacketRequest(Request.DataVersion, Request.DataOffset);
+		}
+	}
+}
+
+void AGameplayDebuggerCategoryReplicator::OnReceivedDataPackPacket(int32 CategoryId, int32 DataPackId, const FGameplayDebuggerDataPack& DataPacket)
+{
+	if (Categories.IsValidIndex(CategoryId) && Categories[CategoryId]->ReplicatedDataPacks.IsValidIndex(DataPackId))
+	{
+		FGameplayDebuggerDataPack& DataPack = Categories[CategoryId]->ReplicatedDataPacks[DataPackId];
+		bool bIsPacketValid = true;
+
+		if (DataPack.Header.DataVersion != DataPacket.Header.DataVersion)
+		{
+			// new content of data pack:
+			if (DataPacket.Header.DataOffset == 0)
+			{
+				// first packet of data, replace old data pack's intermediate data
+				DataPack.Header = DataPacket.Header;
+				DataPack.Data = DataPacket.Data;
+			}
+			else
+			{
+				// somewhere in the middle, discard packet
+				UE_LOG(LogGameplayDebugReplication, Verbose, TEXT("Error: received packet from the middle of content with different version, discarding! (category[%d]:%s, packet.DataVersion:%d packet.DataOffset:%d, data[%d].DataVersion:%d)"),
+					CategoryId, *Categories[CategoryId]->GetCategoryName().ToString(),
+					DataPacket.Header.DataVersion,
+					DataPacket.Header.DataOffset,
+					DataPackId, DataPack.Header.DataVersion,
+					DataPackId, DataPack.Header.DataOffset);
+
+				bIsPacketValid = false;
+			}
+		}
+		else
+		{
+			// another packet for existing data pack
+			if (DataPacket.Header.DataOffset == DataPack.Data.Num())
+			{
+				// offset matches, this is next expected packet
+				DataPack.Data.Append(DataPacket.Data);
+				DataPack.Header.DataOffset = DataPack.Data.Num();
+			}
+			else
+			{
+				// offset mismatch, discard packet
+				UE_LOG(LogGameplayDebugReplication, Verbose, TEXT("Error: received packet doesn't match expected chunk, discarding! (category[%d]:%s, packet.DataOffset:%d, data[%d].DataOffset:%d data[%d].Data.Num:%d)"),
+					CategoryId, *Categories[CategoryId]->GetCategoryName().ToString(),
+					DataPacket.Header.DataOffset,
+					DataPackId, DataPack.Header.DataOffset,
+					DataPackId, DataPack.Data.Num());
+
+				bIsPacketValid = false;
+			}
+		}
+
+		// check if data pack is now complete
+		if (bIsPacketValid)
+		{
+			bool bCreatePacketRequest = true;
+			if (DataPack.Data.Num() == DataPack.Header.DataSize)
+			{
+				// complete
+				UE_LOG(LogGameplayDebugReplication, Verbose, TEXT("Category[%d].DataPack[%d] RECEIVED, DataVersion:%d DataSize:%d"),
+					CategoryId, DataPackId, DataPack.Header.DataVersion, DataPack.Header.DataSize);
+
+				// larger data packs will wait on server until receiving confirmation
+				bCreatePacketRequest = (DataPack.Header.DataSize > FGameplayDebuggerDataPack::PacketSize);
+
+				DataPack.OnReplicated();
+				Categories[CategoryId]->OnDataPackReplicated(DataPackId);
+			}
+
+			if (bCreatePacketRequest)
+			{
+				FGameplayDebuggerDataPacketRequest NewRequest;
+				NewRequest.CategoryId = CategoryId;
+				NewRequest.DataPackId = DataPackId;
+				NewRequest.DataVersion = DataPack.Header.DataVersion;
+				NewRequest.DataOffset = DataPack.Data.Num();
+
+				PendingReplicationRequests.Add(NewRequest);
+			}
+		}
+	}
+}
+
+void AGameplayDebuggerCategoryReplicator::TickActor(float DeltaTime, enum ELevelTick TickType, FActorTickFunction& ThisTickFunction)
+{
+	Super::TickActor(DeltaTime, TickType, ThisTickFunction);
+	
+	const float GameTime = GetWorld()->GetTimeSeconds();
+	for (int32 Idx = 0; Idx < Categories.Num(); Idx++)
+	{
+		FGameplayDebuggerCategory& CategoryOb = Categories[Idx].Get();
+		if (CategoryOb.bHasAuthority && CategoryOb.bIsEnabled && ((GameTime - CategoryOb.LastCollectDataTime) > CategoryOb.CollectDataInterval))
+		{
+			// prepare data packs before calling CollectData
+			for (int32 DataPackIdx = 0; DataPackIdx < CategoryOb.ReplicatedDataPacks.Num(); DataPackIdx++)
+			{
+				FGameplayDebuggerDataPack& DataPack = CategoryOb.ReplicatedDataPacks[DataPackIdx];
+				DataPack.bIsDirty = false;
+
+				if (DataPack.Flags == EGameplayDebuggerDataPack::AutoResetData)
+				{
+					DataPack.ResetDelegate.Execute();
+				}
+			}
+
+			CategoryOb.ReplicatedLines.Reset();
+			CategoryOb.ReplicatedShapes.Reset();
+
+			CategoryOb.CollectData(OwnerPC, DebugActor);
+
+			// update dirty data packs
+			for (int32 DataPackIdx = 0; DataPackIdx < CategoryOb.ReplicatedDataPacks.Num(); DataPackIdx++)
+			{
+				FGameplayDebuggerDataPack& DataPack = CategoryOb.ReplicatedDataPacks[DataPackIdx];
+				if (CategoryOb.bIsLocal)
+				{
+					const bool bWasDirty = DataPack.CheckDirtyAndUpdate();
+					if (bWasDirty)
+					{
+						CategoryOb.OnDataPackReplicated(DataPackIdx);
+					}
+				}
+				else
+				{
+					const bool bWasDirty = DataPack.RequestReplication();
+					if (bWasDirty)
+					{
+						UE_LOG(LogGameplayDebugReplication, Verbose, TEXT("Category[%d].DataPack[%d] SENT, DataVersion:%d DataSize:%d"),
+							Idx, DataPackIdx, DataPack.Header.DataVersion, DataPack.Header.DataSize);
+					}
+				}
+			}
+		}
+	}
+}
+
+void AGameplayDebuggerCategoryReplicator::SetReplicatorOwner(APlayerController* InOwnerPC)
+{
+	if (!bIsEnabled)
+	{
+		// can't use bHasAuthority, BeginPlay was not called yet
+		UWorld* World = GetWorld();
+		const ENetMode NetMode = World->GetNetMode();
+
+		if (NetMode != NM_Client)
+		{
+			OwnerPC = InOwnerPC;
+		}
+	}
+}
+
+void AGameplayDebuggerCategoryReplicator::SetEnabled(bool bEnable)
+{
+	if (bHasAuthority)
+	{
+		bIsEnabled = bEnable;
+		SetActorTickEnabled(bEnable);
+	}
+	else
+	{
+		ServerSetEnabled(bEnable);
+	}
+
+	MarkComponentsRenderStateDirty();
+}
+
+void AGameplayDebuggerCategoryReplicator::SetDebugActor(AActor* Actor)
+{
+	if (bHasAuthority)
+	{
+		DebugActor = Actor;
+	}
+	else
+	{
+		ServerSetDebugActor(Actor);
+	}
+}
+
+void AGameplayDebuggerCategoryReplicator::SetCategoryEnabled(int32 CategoryId, bool bEnable)
+{
+	if (bHasAuthority)
+	{
+		if (Categories.IsValidIndex(CategoryId))
+		{
+			Categories[CategoryId]->bIsEnabled = bEnable;
+		}
+	}
+	else
+	{
+		ServerSetCategoryEnabled(CategoryId, bEnable);
+	}
+
+	MarkComponentsRenderStateDirty();
+}
+
+bool AGameplayDebuggerCategoryReplicator::IsCategoryEnabled(int32 CategoryId) const
+{
+	return Categories.IsValidIndex(CategoryId) && Categories[CategoryId]->IsCategoryEnabled(); 
+}
diff --git a/Engine/Source/Developer/GameplayDebugger/Private/GameplayDebuggerConfig.cpp b/Engine/Source/Developer/GameplayDebugger/Private/GameplayDebuggerConfig.cpp
new file mode 100644
--- /dev/null
+++ b/Engine/Source/Developer/GameplayDebugger/Private/GameplayDebuggerConfig.cpp
@@ -0,0 +1,40 @@
+// Copyright 1998-2016 Epic Games, Inc. All Rights Reserved.
+
+#include "GameplayDebuggerPrivatePCH.h"
+#include "GameplayDebuggerConfig.h"
+#include "GameplayDebuggerTypes.h"
+
+UGameplayDebuggerConfig::UGameplayDebuggerConfig(const FObjectInitializer& ObjectInitializer) : Super(ObjectInitializer)
+{
+	ActivationKey = EKeys::Apostrophe;
+
+	CategoryRowNextKey = EKeys::Add;
+	CategoryRowPrevKey = EKeys::Subtract;
+
+	CategorySlot0 = EKeys::NumPadZero;
+	CategorySlot1 = EKeys::NumPadOne;
+	CategorySlot2 = EKeys::NumPadTwo;
+	CategorySlot3 = EKeys::NumPadThree;
+	CategorySlot4 = EKeys::NumPadFour;
+	CategorySlot5 = EKeys::NumPadFive;
+	CategorySlot6 = EKeys::NumPadSix;
+	CategorySlot7 = EKeys::NumPadSeven;
+	CategorySlot8 = EKeys::NumPadEight;
+	CategorySlot9 = EKeys::NumPadNine;
+
+	GameHUDKey = EKeys::Tilde;
+	GameHUDModifiers.bCtrl = true;
+	bUseGameHUDModifiers = true;
+	DebugMessagesKey = EKeys::Tab;
+	DebugMessagesModifiers.bCtrl = true;
+	bUseDebugMessagesModifiers = true;
+	SpectatorKey = EKeys::Tab;
+
+	bEnableExtension_GameHUD = true;
+	bEnableExtension_Spectator = true;
+
+	DebugCanvasPaddingLeft = 10.0f;
+	DebugCanvasPaddingRight = 10.0f;
+	DebugCanvasPaddingTop = 10.0f;
+	DebugCanvasPaddingBottom = 10.0f;
+}
diff --git a/Engine/Source/Developer/GameplayDebugger/Private/GameplayDebuggerLocalController.cpp b/Engine/Source/Developer/GameplayDebugger/Private/GameplayDebuggerLocalController.cpp
new file mode 100644
--- /dev/null
+++ b/Engine/Source/Developer/GameplayDebugger/Private/GameplayDebuggerLocalController.cpp
@@ -0,0 +1,683 @@
+// Copyright 1998-2016 Epic Games, Inc. All Rights Reserved.
+
+#include "GameplayDebuggerPrivatePCH.h"
+#include "GameplayDebuggerLocalController.h"
+#include "GameplayDebuggerPlayerManager.h"
+#include "GameplayDebuggerAddonManager.h"
+#include "GameplayDebuggerExtension.h"
+#include "GameplayDebuggerConfig.h"
+#include "Debug/DebugDrawService.h"
+#include "Engine/Selection.h"
+#include "Engine/Canvas.h"
+#include "GameFramework/PlayerInput.h"
+#include "GameFramework/Pawn.h"
+
+UGameplayDebuggerLocalController::UGameplayDebuggerLocalController(const FObjectInitializer& ObjectInitializer) : Super(ObjectInitializer)
+{
+	bSimulateMode = false;
+	bNeedsCleanup = false;
+	bIsSelectingActor = false;
+	bIsLocallyEnabled = false;
+	bPrevLocallyEnabled = false;
+	ActiveRowIdx = 0;
+}
+
+void UGameplayDebuggerLocalController::Initialize(AGameplayDebuggerCategoryReplicator& Replicator, AGameplayDebuggerPlayerManager& Manager)
+{
+	CachedReplicator = &Replicator;
+	CachedPlayerManager = &Manager;
+	bSimulateMode = FGameplayDebuggerAddonBase::IsSimulateInEditor();
+
+	UDebugDrawService::Register(bSimulateMode ? TEXT("DebugAI") : TEXT("Game"), FDebugDrawDelegate::CreateUObject(this, &UGameplayDebuggerLocalController::OnDebugDraw));
+	if (bSimulateMode)
+	{
+		USelection::SelectionChangedEvent.AddUObject(this, &UGameplayDebuggerLocalController::OnSelectionChanged);
+	}
+
+	FGameplayDebuggerAddonManager& AddonManager = FGameplayDebuggerAddonManager::GetCurrent();
+	AddonManager.OnExtensionsChanged.AddUObject(this, &UGameplayDebuggerLocalController::OnExtensionsChanged);
+	AddonManager.OnCategoriesChanged.AddUObject(this, &UGameplayDebuggerLocalController::OnCategoriesChanged);
+	AddonManager.CreateExtensions(Replicator, Extensions);
+	OnCategoriesChanged();
+
+	const UGameplayDebuggerConfig* SettingsCDO = UGameplayDebuggerConfig::StaticClass()->GetDefaultObject<UGameplayDebuggerConfig>();
+	const FKey NumpadKeys[] = { EKeys::NumPadZero, EKeys::NumPadOne, EKeys::NumPadTwo, EKeys::NumPadThree, EKeys::NumPadFour,
+		EKeys::NumPadFive, EKeys::NumPadSix, EKeys::NumPadSeven, EKeys::NumPadEight, EKeys::NumPadNine };
+	const FKey CategorySlots[] = { SettingsCDO->CategorySlot0, SettingsCDO->CategorySlot1, SettingsCDO->CategorySlot2, SettingsCDO->CategorySlot3, SettingsCDO->CategorySlot4,
+		SettingsCDO->CategorySlot5, SettingsCDO->CategorySlot6, SettingsCDO->CategorySlot7, SettingsCDO->CategorySlot8, SettingsCDO->CategorySlot9 };
+
+	bool bIsNumpadOnly = true;
+	for (int32 Idx = 0; Idx < ARRAY_COUNT(CategorySlots); Idx++)
+	{
+		bool bHasPattern = false;
+		for (int32 PatternIdx = 0; PatternIdx < ARRAY_COUNT(NumpadKeys); PatternIdx++)
+		{
+			if (CategorySlots[Idx] == NumpadKeys[PatternIdx])
+			{
+				bHasPattern = true;
+				break;
+			}
+		}
+
+		if (!bHasPattern)
+		{
+			bIsNumpadOnly = false;
+			break;
+		}
+	}
+
+	ActivationKeyDesc = GetKeyDescriptionLong(SettingsCDO->ActivationKey);
+	RowUpKeyDesc = GetKeyDescriptionShort(SettingsCDO->CategoryRowPrevKey);
+	RowDownKeyDesc = GetKeyDescriptionShort(SettingsCDO->CategoryRowNextKey);
+	CategoryKeysDesc = bIsNumpadOnly ? TEXT("{yellow}Numpad{white}") : TEXT("highlighted keys");
+
+	PaddingLeft = SettingsCDO->DebugCanvasPaddingLeft;
+	PaddingRight = SettingsCDO->DebugCanvasPaddingRight;
+	PaddingTop = SettingsCDO->DebugCanvasPaddingTop;
+	PaddingBottom = SettingsCDO->DebugCanvasPaddingBottom;
+
+	bNeedsCleanup = true;
+}
+
+void UGameplayDebuggerLocalController::Cleanup()
+{
+	if (bSimulateMode)
+	{
+		USelection::SelectionChangedEvent.RemoveAll(this);
+	}
+
+	bNeedsCleanup = false;
+}
+
+void UGameplayDebuggerLocalController::BeginDestroy()
+{
+	Super::BeginDestroy();
+	if (bNeedsCleanup)
+	{
+		Cleanup();
+	}
+}
+
+void UGameplayDebuggerLocalController::OnDebugDraw(class UCanvas* Canvas, class APlayerController* PC)
+{
+	if (CachedReplicator && CachedReplicator->IsEnabled())
+	{
+		FGameplayDebuggerCanvasContext CanvasContext(Canvas, GEngine->GetSmallFont());
+		CanvasContext.CursorX = CanvasContext.DefaultX = PaddingLeft;
+		CanvasContext.CursorY = CanvasContext.DefaultY = PaddingTop;
+
+		if (bSimulateMode)
+		{
+			DrawSimulateHeader(CanvasContext);
+		}
+		else
+		{
+			DrawHeader(CanvasContext);
+		}
+
+		for (int32 Idx = 0; Idx < NumCategories; Idx++)
+		{
+			TSharedRef<FGameplayDebuggerCategory> Category = CachedReplicator->GetCategory(Idx);
+			if (Category->ShouldDrawCategory(CachedReplicator->GetDebugActor()))
+			{
+				Category->DrawCategory(CachedReplicator->GetReplicationOwner(), CanvasContext);
+			}
+		}
+	}
+}
+
+extern RENDERCORE_API FTexture* GWhiteTexture;
+
+void UGameplayDebuggerLocalController::DrawHeader(FGameplayDebuggerCanvasContext& CanvasContext)
+{
+	const int32 NumRows = (NumCategorySlots + 9) / 10;
+	const float LineHeight = CanvasContext.GetLineHeight();
+	const int32 NumExtensionRows = (Extensions.Num() > 0) ? 1 : 0;
+	const float CanvasSizeX = CanvasContext.Canvas->SizeX - PaddingLeft - PaddingRight;
+	
+	const float BackgroundPadding = 5.0f;
+	const float BackgroundPaddingBothSides = BackgroundPadding * 2.0f;
+
+	if (NumRows > 1)
+	{
+		FCanvasTileItem TileItemUpper(FVector2D(0, 0), GWhiteTexture, FVector2D(CanvasSizeX + BackgroundPaddingBothSides, (LineHeight * (ActiveRowIdx + NumExtensionRows + 1)) + BackgroundPadding), FLinearColor(0, 0, 0, 0.2f));
+		FCanvasTileItem ActiveRowTileItem(FVector2D(0, 0), GWhiteTexture, FVector2D(CanvasSizeX + BackgroundPaddingBothSides, LineHeight), FLinearColor(0, 0.5f, 0, 0.3f));
+		FCanvasTileItem TileItemLower(FVector2D(0, 0), GWhiteTexture, FVector2D(CanvasSizeX + BackgroundPaddingBothSides, LineHeight * ((NumRows - ActiveRowIdx - 1)) + BackgroundPadding), FLinearColor(0, 0, 0, 0.2f));
+
+		TileItemUpper.BlendMode = SE_BLEND_Translucent;
+		ActiveRowTileItem.BlendMode = SE_BLEND_Translucent;
+		TileItemLower.BlendMode = SE_BLEND_Translucent;
+
+		CanvasContext.DrawItem(TileItemUpper, PaddingLeft - BackgroundPadding, PaddingTop - BackgroundPadding);
+		CanvasContext.DrawItem(ActiveRowTileItem, PaddingLeft - BackgroundPadding, PaddingTop - BackgroundPadding + TileItemUpper.Size.Y);
+		CanvasContext.DrawItem(TileItemLower, PaddingLeft - BackgroundPadding, PaddingTop - BackgroundPadding + TileItemUpper.Size.Y + ActiveRowTileItem.Size.Y);
+	}
+	else
+	{
+		FCanvasTileItem TileItem(FVector2D(0, 0), GWhiteTexture, FVector2D(CanvasSizeX, LineHeight * (NumRows + NumExtensionRows + 1)) + BackgroundPaddingBothSides, FLinearColor(0, 0, 0, 0.2f));
+		TileItem.BlendMode = SE_BLEND_Translucent;
+		CanvasContext.DrawItem(TileItem, PaddingLeft - BackgroundPadding, PaddingTop - BackgroundPadding);
+	}
+
+	CanvasContext.Printf(TEXT("Tap {yellow}%s{white} to close, use %s to toggle catories."), *ActivationKeyDesc, *CategoryKeysDesc);
+
+	const FString DebugActorDesc = FString::Printf(TEXT("Debug actor: {cyan}%s"), *GetNameSafe(CachedReplicator->GetDebugActor()));
+	float DebugActorSizeX = 0.0f, DebugActorSizeY = 0.0f;
+	CanvasContext.MeasureString(DebugActorDesc, DebugActorSizeX, DebugActorSizeY);
+	CanvasContext.PrintAt(CanvasContext.Canvas->SizeX - PaddingRight - DebugActorSizeX, PaddingTop, DebugActorDesc);
+
+	const FString TimestampDesc = FString::Printf(TEXT("Time: %.2fs"), CachedReplicator->GetWorld()->GetTimeSeconds());
+	float TimestampSizeX = 0.0f, TimestampSizeY = 0.0f;
+	CanvasContext.MeasureString(TimestampDesc, TimestampSizeX, TimestampSizeY);
+	CanvasContext.PrintAt((CanvasSizeX - TimestampSizeX) * 0.5f, PaddingTop, TimestampDesc);
+
+	if (NumRows > 1)
+	{
+		const FString ChangeRowDesc = FString::Printf(TEXT("Prev row: {yellow}%s\n{white}Next row: {yellow}%s"), *RowUpKeyDesc, *RowDownKeyDesc);
+		float RowDescSizeX = 0.0f, RowDescSizeY = 0.0f;
+		CanvasContext.MeasureString(ChangeRowDesc, RowDescSizeX, RowDescSizeY);
+		CanvasContext.PrintAt(CanvasContext.Canvas->SizeX - PaddingRight - RowDescSizeX, PaddingTop + LineHeight * (NumExtensionRows + 1), ChangeRowDesc);
+	}
+
+	if (NumExtensionRows)
+	{
+		FString ExtensionRowDesc;
+		for (int32 ExtensionIdx = 0; ExtensionIdx < Extensions.Num(); ExtensionIdx++)
+		{
+			FString ExtensionDesc = Extensions[ExtensionIdx]->GetDescription();
+			ExtensionDesc.ReplaceInline(TEXT("\n"), TEXT(""));
+
+			if (ExtensionDesc.Len())
+			{
+				if (ExtensionRowDesc.Len())
+				{
+					ExtensionRowDesc += FGameplayDebuggerCanvasStrings::SeparatorSpace;
+				}
+
+				ExtensionRowDesc += ExtensionDesc;
+			}
+		}
+
+		CanvasContext.Print(ExtensionRowDesc);
+	}
+
+	for (int32 RowIdx = 0; RowIdx < NumRows; RowIdx++)
+	{
+		FString CategoryRowDesc;
+		for (int32 Idx = 0; Idx < 10; Idx++)
+		{
+			const int32 CategorySlotIdx = (RowIdx * 10) + Idx;
+			if (SlotCategoryIds.IsValidIndex(CategorySlotIdx) && 
+				SlotNames.IsValidIndex(CategorySlotIdx) &&
+				SlotCategoryIds[CategorySlotIdx].Num())
+			{
+				TSharedRef<FGameplayDebuggerCategory> Category0 = CachedReplicator->GetCategory(SlotCategoryIds[CategorySlotIdx][0]);
+				const bool bIsEnabled = Category0->IsCategoryEnabled();
+				const FString CategoryColorName = (RowIdx == ActiveRowIdx) && (NumRows > 1) ?
+					(bIsEnabled ? *FGameplayDebuggerCanvasStrings::ColorNameEnabledActiveRow : *FGameplayDebuggerCanvasStrings::ColorNameDisabledActiveRow) :
+					(bIsEnabled ? *FGameplayDebuggerCanvasStrings::ColorNameEnabled : *FGameplayDebuggerCanvasStrings::ColorNameDisabled);
+
+				const FString CategoryDesc = (RowIdx == ActiveRowIdx) ?
+					FString::Printf(TEXT("%s{%s}%d:{%s}%s"),
+						Idx ? *FGameplayDebuggerCanvasStrings::SeparatorSpace : TEXT(""),
+						*FGameplayDebuggerCanvasStrings::ColorNameInput,
+						Idx,
+						*CategoryColorName,
+						*SlotNames[CategorySlotIdx]) :
+					FString::Printf(TEXT("%s{%s}%s"),
+						Idx ? *FGameplayDebuggerCanvasStrings::Separator : TEXT(""),
+						*CategoryColorName,
+						*SlotNames[CategorySlotIdx]);
+
+				CategoryRowDesc += CategoryDesc;
+			}
+		}
+
+		CanvasContext.Print(CategoryRowDesc);
+	}
+
+	CanvasContext.DefaultY = CanvasContext.CursorY + LineHeight;
+}
+
+void UGameplayDebuggerLocalController::DrawSimulateHeader(FGameplayDebuggerCanvasContext& CanvasContext)
+{
+	const float LineHeight = CanvasContext.GetLineHeight();
+	const float NewCategoryPosY = (CanvasContext.DefaultY * 2.0f) + (LineHeight * 3.0f);
+	const float BackgroundPadding = 5.0f;
+
+	const FString TimestampDesc = FString::Printf(TEXT("GameplayDebugger time: %.2fs"), CachedReplicator->GetWorld()->GetTimeSeconds());
+	float TimestampSizeX = 0.0f, TimestampSizeY = 0.0f;
+	CanvasContext.MeasureString(TimestampDesc, TimestampSizeX, TimestampSizeY);
+
+	const FString DebugActorDesc = FString::Printf(TEXT("Selected actor: {cyan}%s"), *GetNameSafe(CachedReplicator->GetDebugActor()));
+	float DebugActorSizeX = 0.0f, DebugActorSizeY = 0.0f;
+	CanvasContext.MeasureString(DebugActorDesc, DebugActorSizeX, DebugActorSizeY);
+
+	const float BackgroundSizeX = FMath::Max(DebugActorSizeX, TimestampSizeX) + (BackgroundPadding * 2.f);
+	FCanvasTileItem TileItem(FVector2D(0, 0), GWhiteTexture, FVector2D(BackgroundSizeX, (2.0f * LineHeight) + (2.0f * BackgroundPadding)), FLinearColor(0, 0, 0, 0.2f));
+	TileItem.BlendMode = SE_BLEND_Translucent;
+	CanvasContext.DrawItem(TileItem, (CanvasContext.Canvas->SizeX - BackgroundSizeX) * 0.5f, PaddingTop - BackgroundPadding);
+
+	CanvasContext.PrintAt((CanvasContext.Canvas->SizeX - TimestampSizeX) * 0.5f, PaddingTop, TimestampDesc);
+	CanvasContext.PrintAt((CanvasContext.Canvas->SizeX - DebugActorSizeX) * 0.5f, PaddingTop + LineHeight, DebugActorDesc);
+
+	CanvasContext.DefaultY = CanvasContext.CursorY = NewCategoryPosY;
+}
+
+void UGameplayDebuggerLocalController::BindInput(UInputComponent& InputComponent)
+{
+	TSet<FName> UsedBindings;
+
+	const UGameplayDebuggerConfig* SettingsCDO = UGameplayDebuggerConfig::StaticClass()->GetDefaultObject<UGameplayDebuggerConfig>();
+	InputComponent.BindKey(SettingsCDO->ActivationKey, IE_Pressed, this, &UGameplayDebuggerLocalController::OnActivationPressed);
+	InputComponent.BindKey(SettingsCDO->ActivationKey, IE_Released, this, &UGameplayDebuggerLocalController::OnActivationReleased);
+	UsedBindings.Add(SettingsCDO->ActivationKey.GetFName());
+
+	if (bIsLocallyEnabled)
+	{
+		InputComponent.BindKey(SettingsCDO->CategorySlot0, IE_Pressed, this, &UGameplayDebuggerLocalController::OnCategory0Pressed);
+		InputComponent.BindKey(SettingsCDO->CategorySlot1, IE_Pressed, this, &UGameplayDebuggerLocalController::OnCategory1Pressed);
+		InputComponent.BindKey(SettingsCDO->CategorySlot2, IE_Pressed, this, &UGameplayDebuggerLocalController::OnCategory2Pressed);
+		InputComponent.BindKey(SettingsCDO->CategorySlot3, IE_Pressed, this, &UGameplayDebuggerLocalController::OnCategory3Pressed);
+		InputComponent.BindKey(SettingsCDO->CategorySlot4, IE_Pressed, this, &UGameplayDebuggerLocalController::OnCategory4Pressed);
+		InputComponent.BindKey(SettingsCDO->CategorySlot5, IE_Pressed, this, &UGameplayDebuggerLocalController::OnCategory5Pressed);
+		InputComponent.BindKey(SettingsCDO->CategorySlot6, IE_Pressed, this, &UGameplayDebuggerLocalController::OnCategory6Pressed);
+		InputComponent.BindKey(SettingsCDO->CategorySlot7, IE_Pressed, this, &UGameplayDebuggerLocalController::OnCategory7Pressed);
+		InputComponent.BindKey(SettingsCDO->CategorySlot8, IE_Pressed, this, &UGameplayDebuggerLocalController::OnCategory8Pressed);
+		InputComponent.BindKey(SettingsCDO->CategorySlot9, IE_Pressed, this, &UGameplayDebuggerLocalController::OnCategory9Pressed);
+
+		InputComponent.BindKey(SettingsCDO->CategoryRowPrevKey, IE_Pressed, this, &UGameplayDebuggerLocalController::OnCategoryRowUpPressed);
+		InputComponent.BindKey(SettingsCDO->CategoryRowNextKey, IE_Pressed, this, &UGameplayDebuggerLocalController::OnCategoryRowDownPressed);
+
+		UsedBindings.Add(SettingsCDO->CategorySlot0.GetFName());
+		UsedBindings.Add(SettingsCDO->CategorySlot1.GetFName());
+		UsedBindings.Add(SettingsCDO->CategorySlot2.GetFName());
+		UsedBindings.Add(SettingsCDO->CategorySlot3.GetFName());
+		UsedBindings.Add(SettingsCDO->CategorySlot4.GetFName());
+		UsedBindings.Add(SettingsCDO->CategorySlot5.GetFName());
+		UsedBindings.Add(SettingsCDO->CategorySlot6.GetFName());
+		UsedBindings.Add(SettingsCDO->CategorySlot7.GetFName());
+		UsedBindings.Add(SettingsCDO->CategorySlot8.GetFName());
+		UsedBindings.Add(SettingsCDO->CategorySlot9.GetFName());
+		UsedBindings.Add(SettingsCDO->CategoryRowPrevKey.GetFName());
+		UsedBindings.Add(SettingsCDO->CategoryRowNextKey.GetFName());
+
+		for (int32 Idx = 0; Idx < NumCategories; Idx++)
+		{
+			TSharedRef<FGameplayDebuggerCategory> Category = CachedReplicator->GetCategory(Idx);
+			const int32 NumInputHandlers = Category->GetNumInputHandlers();
+
+			for (int32 HandlerIdx = 0; HandlerIdx < NumInputHandlers; HandlerIdx++)
+			{
+				FGameplayDebuggerInputHandler& HandlerData = Category->GetInputHandler(HandlerIdx);
+				if (HandlerData.Modifier.bPressed || HandlerData.Modifier.bReleased)
+				{
+					FInputChord InputChord(FKey(HandlerData.KeyName), HandlerData.Modifier.bShift, HandlerData.Modifier.bCtrl, HandlerData.Modifier.bAlt, HandlerData.Modifier.bCmd);
+					FInputKeyBinding NewBinding(InputChord, HandlerData.Modifier.bPressed ? IE_Pressed : IE_Released);
+					NewBinding.KeyDelegate.GetDelegateForManualSet().BindUObject(this, &UGameplayDebuggerLocalController::OnCategoryBindingEvent, Idx, HandlerIdx);
+
+					InputComponent.KeyBindings.Add(NewBinding);
+					UsedBindings.Add(HandlerData.KeyName);
+				}
+			}
+		}
+
+		for (int32 Idx = 0; Idx < Extensions.Num(); Idx++)
+		{
+			TSharedRef<FGameplayDebuggerExtension> Extension = Extensions[Idx];
+			const int32 NumInputHandlers = Extension->GetNumInputHandlers();
+
+			for (int32 HandlerIdx = 0; HandlerIdx < NumInputHandlers; HandlerIdx++)
+			{
+				FGameplayDebuggerInputHandler& HandlerData = Extension->GetInputHandler(HandlerIdx);
+				if (HandlerData.Modifier.bPressed || HandlerData.Modifier.bReleased)
+				{
+					FInputChord InputChord(FKey(HandlerData.KeyName), HandlerData.Modifier.bShift, HandlerData.Modifier.bCtrl, HandlerData.Modifier.bAlt, HandlerData.Modifier.bCmd);
+					FInputKeyBinding NewBinding(InputChord, HandlerData.Modifier.bPressed ? IE_Pressed : IE_Released);
+					NewBinding.KeyDelegate.GetDelegateForManualSet().BindUObject(this, &UGameplayDebuggerLocalController::OnExtensionBindingEvent, Idx, HandlerIdx);
+
+					InputComponent.KeyBindings.Add(NewBinding);
+					UsedBindings.Add(HandlerData.KeyName);
+				}
+			}
+		}
+	}
+
+	if (CachedReplicator && CachedReplicator->GetReplicationOwner() && CachedReplicator->GetReplicationOwner()->PlayerInput)
+	{
+		TSet<FName> RemovedMasks = MaskedDebugBindings.Difference(UsedBindings);
+		TSet<FName> AddedMasks = UsedBindings.Difference(MaskedDebugBindings);
+
+		UPlayerInput* Input = CachedReplicator->GetReplicationOwner()->PlayerInput;
+		for (int32 Idx = 0; Idx < Input->DebugExecBindings.Num(); Idx++)
+		{
+			FKeyBind& DebugBinding = Input->DebugExecBindings[Idx];
+			const bool bRemoveMask = RemovedMasks.Contains(DebugBinding.Key.GetFName());
+			const bool bAddMask = AddedMasks.Contains(DebugBinding.Key.GetFName());
+
+			if (bAddMask || bRemoveMask)
+			{
+				DebugBinding.bDisabled = bAddMask;
+			}
+		}
+
+		MaskedDebugBindings = UsedBindings;
+	}
+}
+
+void UGameplayDebuggerLocalController::OnActivationPressed()
+{
+	bPrevLocallyEnabled = bIsLocallyEnabled;
+	if (CachedReplicator)
+	{
+		CachedReplicator->GetWorldTimerManager().SetTimer(StartSelectingActorHandle, this, &UGameplayDebuggerLocalController::OnStartSelectingActor, 0.2f);
+	}
+}
+
+void UGameplayDebuggerLocalController::OnActivationReleased()
+{
+	if (CachedReplicator)
+	{
+		UWorld* World = CachedReplicator->GetWorld();
+		if (StartSelectingActorHandle.IsValid())
+		{
+			bIsLocallyEnabled = !CachedReplicator->IsEnabled();
+			CachedReplicator->SetEnabled(bIsLocallyEnabled);
+
+			if (bIsLocallyEnabled)
+			{
+				DebugActorCandidate = nullptr;
+				OnSelectActorTick();
+			}
+		}
+
+		World->GetTimerManager().ClearTimer(StartSelectingActorHandle);
+		World->GetTimerManager().ClearTimer(SelectActorTickHandle);
+	}
+
+	StartSelectingActorHandle.Invalidate();
+	SelectActorTickHandle.Invalidate();
+	bIsSelectingActor = false;
+
+	if (bPrevLocallyEnabled != bIsLocallyEnabled)
+	{
+		if (bIsLocallyEnabled)
+		{
+			NotifyExtensionsActivation();
+		}
+		else
+		{
+			NotifyExtensionsDeactivation();
+		}
+
+		CachedPlayerManager->RefreshInputBindings(*CachedReplicator);
+	}
+}
+
+void UGameplayDebuggerLocalController::OnCategory0Pressed()
+{
+	ToggleSlotState((ActiveRowIdx * 10) + 0);
+}
+
+void UGameplayDebuggerLocalController::OnCategory1Pressed()
+{
+	ToggleSlotState((ActiveRowIdx * 10) + 1);
+}
+
+void UGameplayDebuggerLocalController::OnCategory2Pressed()
+{
+	ToggleSlotState((ActiveRowIdx * 10) + 2);
+}
+
+void UGameplayDebuggerLocalController::OnCategory3Pressed()
+{
+	ToggleSlotState((ActiveRowIdx * 10) + 3);
+}
+
+void UGameplayDebuggerLocalController::OnCategory4Pressed()
+{
+	ToggleSlotState((ActiveRowIdx * 10) + 4);
+}
+
+void UGameplayDebuggerLocalController::OnCategory5Pressed()
+{
+	ToggleSlotState((ActiveRowIdx * 10) + 5);
+}
+
+void UGameplayDebuggerLocalController::OnCategory6Pressed()
+{
+	ToggleSlotState((ActiveRowIdx * 10) + 6);
+}
+
+void UGameplayDebuggerLocalController::OnCategory7Pressed()
+{
+	ToggleSlotState((ActiveRowIdx * 10) + 7);
+}
+
+void UGameplayDebuggerLocalController::OnCategory8Pressed()
+{
+	ToggleSlotState((ActiveRowIdx * 10) + 8);
+}
+
+void UGameplayDebuggerLocalController::OnCategory9Pressed()
+{
+	ToggleSlotState((ActiveRowIdx * 10) + 9);
+}
+
+void UGameplayDebuggerLocalController::OnCategoryRowUpPressed()
+{
+	const int32 NumRows = (NumCategorySlots + 9) / 10;
+	ActiveRowIdx = (NumRows > 1) ? ((ActiveRowIdx + NumRows - 1) % NumRows) : 0;
+}
+
+void UGameplayDebuggerLocalController::OnCategoryRowDownPressed()
+{
+	const int32 NumRows = (NumCategorySlots + 9) / 10;
+	ActiveRowIdx = (NumRows > 1) ? ((ActiveRowIdx + 1) % NumRows) : 0;
+}
+
+void UGameplayDebuggerLocalController::OnCategoryBindingEvent(int32 CategoryId, int32 HandlerId)
+{
+	if (CachedReplicator)
+	{
+		TSharedRef<FGameplayDebuggerCategory> Category = CachedReplicator->GetCategory(CategoryId);
+		if (Category->IsCategoryEnabled())
+		{
+			FGameplayDebuggerInputHandler& InputHandler = Category->GetInputHandler(HandlerId);
+			InputHandler.Delegate.ExecuteIfBound();
+		}
+	}
+}
+
+void UGameplayDebuggerLocalController::OnExtensionBindingEvent(int32 ExtensionId, int32 HandlerId)
+{
+	if (Extensions.IsValidIndex(ExtensionId))
+	{
+		TSharedRef<FGameplayDebuggerExtension> Extension = Extensions[ExtensionId];
+		FGameplayDebuggerInputHandler& InputHandler = Extension->GetInputHandler(HandlerId);
+		InputHandler.Delegate.ExecuteIfBound();
+	}
+}
+
+void UGameplayDebuggerLocalController::OnStartSelectingActor()
+{
+	StartSelectingActorHandle.Invalidate();
+	if (CachedReplicator)
+	{
+		if (!CachedReplicator->IsEnabled())
+		{
+			bIsLocallyEnabled = true;
+			CachedReplicator->SetEnabled(bIsLocallyEnabled);
+		}
+
+		bIsSelectingActor = true;
+		DebugActorCandidate = nullptr;
+
+		const bool bLooping = true;
+		CachedReplicator->GetWorldTimerManager().SetTimer(SelectActorTickHandle, this, &UGameplayDebuggerLocalController::OnSelectActorTick, 0.01f, bLooping);
+
+		OnSelectActorTick();
+	}
+}
+
+void UGameplayDebuggerLocalController::OnSelectActorTick()
+{
+	APlayerController* OwnerPC = CachedReplicator ? CachedReplicator->GetReplicationOwner() : nullptr;
+	if (OwnerPC)
+	{
+		FVector CameraLocation;
+		FRotator CameraRotation;
+		OwnerPC->GetPlayerViewPoint(CameraLocation, CameraRotation);
+
+		AActor* BestCandidate = nullptr;
+		float BestScore = 0.0f;
+		
+		// TODO: move to module's settings
+		const float MaxScanDistance = 25000.0f;
+
+		const FVector ViewDir = CameraRotation.Vector();
+		for (FConstPawnIterator It = OwnerPC->GetWorld()->GetPawnIterator(); It; ++It)
+		{
+			APawn* TestPawn = *It;
+			if (TestPawn && !TestPawn->bHidden && TestPawn->GetActorEnableCollision() && TestPawn != OwnerPC->GetPawn())
+			{
+				FVector DirToPawn = (TestPawn->GetActorLocation() - CameraLocation);
+				float DistToPawn = DirToPawn.Size();
+				if (FMath::IsNearlyZero(DistToPawn))
+				{
+					DirToPawn = ViewDir;
+					DistToPawn = 1.0f;
+				}
+				else
+				{
+					DirToPawn /= DistToPawn;
+				}
+
+				const float ViewDot = FVector::DotProduct(ViewDir, DirToPawn);
+				if (DistToPawn < MaxScanDistance && ViewDot > BestScore)
+				{
+					BestScore = ViewDot;
+					BestCandidate = TestPawn;
+				}
+			}
+		}
+
+		// cache to avoid multiple RPC with the same actor
+		if (DebugActorCandidate != BestCandidate)
+		{
+			DebugActorCandidate = BestCandidate;
+			CachedReplicator->SetDebugActor(BestCandidate);
+		}
+	}
+}
+
+void UGameplayDebuggerLocalController::ToggleSlotState(int32 SlotIdx)
+{
+	if (CachedReplicator && SlotCategoryIds.IsValidIndex(SlotIdx) && SlotCategoryIds[SlotIdx].Num())
+	{
+		const bool bIsEnabled = CachedReplicator->IsCategoryEnabled(SlotCategoryIds[SlotIdx][0]);
+		for (int32 Idx = 0; Idx < SlotCategoryIds[SlotIdx].Num(); Idx++)
+		{
+			const int32 CategoryId = SlotCategoryIds[SlotIdx][Idx];
+			CachedReplicator->SetCategoryEnabled(CategoryId, !bIsEnabled);
+		}
+	}
+}
+
+FString UGameplayDebuggerLocalController::GetKeyDescriptionShort(const FKey& KeyBind) const
+{
+	return FString::Printf(TEXT("[%s]"), *KeyBind.GetFName().ToString());
+}
+
+FString UGameplayDebuggerLocalController::GetKeyDescriptionLong(const FKey& KeyBind) const
+{
+	const FString KeyDisplay = KeyBind.GetDisplayName().ToString();
+	const FString KeyName = KeyBind.GetFName().ToString();
+	return (KeyDisplay == KeyName) ? FString::Printf(TEXT("[%s]"), *KeyDisplay) : FString::Printf(TEXT("%s [%s key])"), *KeyDisplay, *KeyName);
+}
+
+void UGameplayDebuggerLocalController::OnSelectionChanged(UObject* Object)
+{
+	USelection* Selection = Cast<USelection>(Object);
+	if (Selection && CachedReplicator)
+	{
+		APawn* SelectedPawn = nullptr;
+		for (int32 Idx = 0; Idx < Selection->Num(); Idx++)
+		{
+			AController* SelectedController = Cast<AController>(Selection->GetSelectedObject(Idx));
+			SelectedPawn = SelectedController ? SelectedController->GetPawn() : Cast<APawn>(Selection->GetSelectedObject(Idx));
+			if (SelectedPawn)
+			{
+				break;
+			}
+		}
+
+		CachedReplicator->SetDebugActor(SelectedPawn);
+	}
+}
+
+void UGameplayDebuggerLocalController::NotifyExtensionsActivation()
+{
+	for (int32 Idx = 0; Idx < Extensions.Num(); Idx++)
+	{
+		Extensions[Idx]->OnActivated();
+	}
+}
+
+void UGameplayDebuggerLocalController::NotifyExtensionsDeactivation()
+{
+	for (int32 Idx = 0; Idx < Extensions.Num(); Idx++)
+	{
+		Extensions[Idx]->OnDeactivated();
+	}
+}
+
+void UGameplayDebuggerLocalController::OnExtensionsChanged()
+{
+	if (CachedReplicator == nullptr || CachedPlayerManager == nullptr)
+	{
+		return;
+	}
+
+	if (bIsLocallyEnabled)
+	{
+		NotifyExtensionsDeactivation();
+	}
+
+	Extensions.Reset();
+	
+	FGameplayDebuggerAddonManager& AddonManager = FGameplayDebuggerAddonManager::GetCurrent();
+	AddonManager.CreateExtensions(*CachedReplicator, Extensions);
+
+	if (bIsLocallyEnabled)
+	{
+		NotifyExtensionsActivation();
+		CachedPlayerManager->RefreshInputBindings(*CachedReplicator);
+	}
+}
+
+void UGameplayDebuggerLocalController::OnCategoriesChanged()
+{
+	FGameplayDebuggerAddonManager& AddonManager = FGameplayDebuggerAddonManager::GetCurrent();
+
+	SlotNames.Reset();
+	SlotNames.Append(AddonManager.GetSlotNames());
+
+	// categories are already sorted using AddonManager.SlotMap, build Slot to Category Id map accordingly
+	const TArray< TArray<int32> >& SlotMap = AddonManager.GetSlotMap();
+	SlotCategoryIds.Reset();
+	SlotCategoryIds.AddDefaulted(SlotMap.Num());
+
+	int32 CategoryId = 0;
+	for (int32 SlotIdx = 0; SlotIdx < SlotMap.Num(); SlotIdx++)
+	{
+		for (int32 InnerIdx = 0; InnerIdx < SlotMap[SlotIdx].Num(); InnerIdx++)
+		{
+			SlotCategoryIds[SlotIdx].Add(CategoryId);
+			CategoryId++;
+		}
+	}
+
+	NumCategorySlots = SlotCategoryIds.Num();
+	NumCategories = AddonManager.GetNumCategories();
+}
diff --git a/Engine/Source/Developer/GameplayDebugger/Private/GameplayDebuggerPrivatePCH.h b/Engine/Source/Developer/GameplayDebugger/Private/GameplayDebuggerPrivatePCH.h
new file mode 100644
--- /dev/null
+++ b/Engine/Source/Developer/GameplayDebugger/Private/GameplayDebuggerPrivatePCH.h
@@ -0,0 +1,25 @@
+// Copyright 1998-2016 Epic Games, Inc. All Rights Reserved.
+
+#pragma once
+
+#include "Core.h"
+#include "CoreUObject.h"
+#include "InputCore.h"
+#include "EngineDefines.h"
+#include "EngineSettings.h"
+#include "EngineStats.h"
+#include "EngineLogs.h"
+#include "EngineGlobals.h"
+
+#include "Engine/EngineBaseTypes.h"
+#include "Engine/EngineTypes.h"
+#include "Engine/Engine.h"
+#include "Components/ActorComponent.h"
+#include "GameFramework/Actor.h"
+#include "GameFramework/Pawn.h"
+#include "GameFramework/PlayerController.h"
+
+#include "GameplayDebuggerTypes.h"
+#include "GameplayDebuggerCategory.h"
+#include "GameplayDebuggerCategoryReplicator.h"
+
diff --git a/Engine/Source/Developer/GameplayDebugger/Private/GameplayDebuggerRenderingComponent.cpp b/Engine/Source/Developer/GameplayDebugger/Private/GameplayDebuggerRenderingComponent.cpp
new file mode 100644
--- /dev/null
+++ b/Engine/Source/Developer/GameplayDebugger/Private/GameplayDebuggerRenderingComponent.cpp
@@ -0,0 +1,156 @@
+// Copyright 1998-2016 Epic Games, Inc. All Rights Reserved.
+
+#include "GameplayDebuggerPrivatePCH.h"
+#include "GameplayDebuggerRenderingComponent.h"
+#include "DebugRenderSceneProxy.h"
+
+//////////////////////////////////////////////////////////////////////////
+// FGameplayDebuggerCompositeSceneProxy
+
+class FGameplayDebuggerCompositeSceneProxy : public FDebugRenderSceneProxy
+{
+public:
+	FGameplayDebuggerCompositeSceneProxy(const UPrimitiveComponent* InComponent) : FDebugRenderSceneProxy(InComponent) { }
+
+	virtual ~FGameplayDebuggerCompositeSceneProxy()
+	{
+		for (int32 Idx = 0; Idx < ChildProxies.Num(); Idx++)
+		{
+			delete ChildProxies[Idx];
+		}
+	}
+
+	virtual void DrawStaticElements(FStaticPrimitiveDrawInterface* PDI) override
+	{
+		for (int32 Idx = 0; Idx < ChildProxies.Num(); Idx++)
+		{
+			ChildProxies[Idx]->DrawStaticElements(PDI);
+		}
+	}
+
+	virtual void GetDynamicMeshElements(const TArray<const FSceneView*>& Views, const FSceneViewFamily& ViewFamily, uint32 VisibilityMap, FMeshElementCollector& Collector) const override
+	{
+		for (int32 Idx = 0; Idx < ChildProxies.Num(); Idx++)
+		{
+			ChildProxies[Idx]->GetDynamicMeshElements(Views, ViewFamily, VisibilityMap, Collector);
+		}
+	}
+
+	virtual void RegisterDebugDrawDelgate() override
+	{
+		for (int32 Idx = 0; Idx < ChildProxies.Num(); Idx++)
+		{
+			ChildProxies[Idx]->RegisterDebugDrawDelgate();
+		}
+	}
+
+	virtual void UnregisterDebugDrawDelgate() override
+	{
+		for (int32 Idx = 0; Idx < ChildProxies.Num(); Idx++)
+		{
+			ChildProxies[Idx]->UnregisterDebugDrawDelgate();
+		}
+	}
+
+	virtual FPrimitiveViewRelevance GetViewRelevance(const FSceneView* View) const override
+	{
+		FPrimitiveViewRelevance Result;
+		for (int32 Idx = 0; Idx < ChildProxies.Num(); Idx++)
+		{
+			Result |= ChildProxies[Idx]->GetViewRelevance(View);
+		}
+		return Result;
+	}
+
+	virtual uint32 GetMemoryFootprint(void) const override
+	{
+		return sizeof(*this) + GetAllocatedSize();
+	}
+
+	uint32 GetAllocatedSize(void) const
+	{
+		uint32 Size = ChildProxies.GetAllocatedSize();
+		for (int32 Idx = 0; Idx < ChildProxies.Num(); Idx++)
+		{
+			Size += ChildProxies[Idx]->GetMemoryFootprint();
+		}
+
+		return Size;
+	}
+
+	void AddChild(FDebugRenderSceneProxy* NewChild)
+	{
+		ChildProxies.AddUnique(NewChild);
+	}
+
+	void AddRange(TArray<FDebugRenderSceneProxy*> Children)
+	{
+		ChildProxies.Append(Children);
+	}
+
+protected:
+	TArray<FDebugRenderSceneProxy*> ChildProxies;
+};
+
+//////////////////////////////////////////////////////////////////////////
+// UGameplayDebuggerRenderingComponent
+
+UGameplayDebuggerRenderingComponent::UGameplayDebuggerRenderingComponent(const FObjectInitializer& ObjInitializer) : Super(ObjInitializer)
+{
+}
+
+FPrimitiveSceneProxy* UGameplayDebuggerRenderingComponent::CreateSceneProxy()
+{
+	FGameplayDebuggerCompositeSceneProxy* CompositeProxy = nullptr;
+
+	AGameplayDebuggerCategoryReplicator* OwnerReplicator = Cast<AGameplayDebuggerCategoryReplicator>(GetOwner());
+	if (OwnerReplicator && OwnerReplicator->IsEnabled())
+	{
+		TArray<FDebugRenderSceneProxy*> SceneProxies;
+		for (int32 Idx = 0; Idx < OwnerReplicator->GetNumCategories(); Idx++)
+		{
+			TSharedRef<FGameplayDebuggerCategory> Category = OwnerReplicator->GetCategory(Idx);
+			if (Category->IsCategoryEnabled())
+			{
+				FDebugRenderSceneProxy* CategorySceneProxy = Category->CreateSceneProxy(this);
+				if (CategorySceneProxy)
+				{
+					SceneProxies.Add(CategorySceneProxy);
+				}
+			}
+		}
+
+		if (SceneProxies.Num())
+		{
+			CompositeProxy = new FGameplayDebuggerCompositeSceneProxy(this);
+			CompositeProxy->AddRange(SceneProxies);
+		}
+	}
+
+	return CompositeProxy;
+}
+
+FBoxSphereBounds UGameplayDebuggerRenderingComponent::CalcBounds(const FTransform &LocalToWorld) const
+{
+	return FBoxSphereBounds(FBox::BuildAABB(FVector::ZeroVector, FVector(1000000.0f, 1000000.0f, 1000000.0f)));
+}
+
+void UGameplayDebuggerRenderingComponent::CreateRenderState_Concurrent()
+{
+	Super::CreateRenderState_Concurrent();
+
+	if (SceneProxy)
+	{
+		static_cast<FGameplayDebuggerCompositeSceneProxy*>(SceneProxy)->RegisterDebugDrawDelgate();
+	}
+}
+
+void UGameplayDebuggerRenderingComponent::DestroyRenderState_Concurrent()
+{
+	if (SceneProxy)
+	{
+		static_cast<FGameplayDebuggerCompositeSceneProxy*>(SceneProxy)->UnregisterDebugDrawDelgate();
+	}
+
+	Super::DestroyRenderState_Concurrent();
+}
diff --git a/Engine/Source/Developer/GameplayDebugger/Private/GameplayDebuggerSettings.cpp b/Engine/Source/Developer/GameplayDebugger/Private/GameplayDebuggerSettings.cpp
--- a/Engine/Source/Developer/GameplayDebugger/Private/GameplayDebuggerSettings.cpp
+++ b/Engine/Source/Developer/GameplayDebugger/Private/GameplayDebuggerSettings.cpp
@@ -1,6 +1,10 @@
 // Copyright 1998-2016 Epic Games, Inc. All Rights Reserved.
 
-#include "GameplayDebuggerPrivate.h"
+//////////////////////////////////////////////////////////////////////////
+// THIS CLASS IS NOW DEPRECATED AND WILL BE REMOVED IN NEXT VERSION
+// Please check GameplayDebugger.h for details.
+
+#include "GameplayDebuggerPrivatePCH.h"
 #include "Misc/CoreMisc.h"
 #include "GameplayDebuggingComponent.h"
 #include "GameplayDebuggingHUDComponent.h"
