Copying //UE4/Dev-Framework to Dev-Main (//UE4/Dev-Main) @ 2879625

diff --git a/Engine/Source/Developer/GameplayDebugger/Private/GameplayDebuggingComponent.cpp b/Engine/Source/Developer/GameplayDebugger/Private/GameplayDebuggingComponent.cpp
--- a/Engine/Source/Developer/GameplayDebugger/Private/GameplayDebuggingComponent.cpp
+++ b/Engine/Source/Developer/GameplayDebugger/Private/GameplayDebuggingComponent.cpp
@@ -1,10 +1,15 @@
 // Copyright 1998-2016 Epic Games, Inc. All Rights Reserved.
-#include "GameplayDebuggerPrivate.h"
+
+//////////////////////////////////////////////////////////////////////////
+// THIS CLASS IS NOW DEPRECATED AND WILL BE REMOVED IN NEXT VERSION
+// Please check GameplayDebugger.h for details.
+
+#include "GameplayDebuggerPrivatePCH.h"
 #include "Net/UnrealNetwork.h"
 #include "DebugRenderSceneProxy.h"
-#include "AI/Navigation/RecastNavMeshGenerator.h"
-#include "NavMeshRenderingHelpers.h"
 #include "AI/Navigation/NavigationSystem.h"
+#include "AI/Navigation/NavMeshRenderingComponent.h"
+#include "AI/Navigation/RecastNavMesh.h"
 #include "EnvironmentQuery/EnvQueryTypes.h"
 #include "EnvironmentQuery/EnvQueryManager.h"
 #include "EnvironmentQuery/EQSRenderingComponent.h"
@@ -20,7 +25,6 @@
 #include "GameFramework/CharacterMovementComponent.h"
 #include "Engine/Channel.h"
 #include "Animation/AnimMontage.h"
-#include "GameplayAbilitiesModule.h"
 #include "Perception/AIPerceptionComponent.h"
 
 #if WITH_EDITOR
@@ -29,6 +33,97 @@
 #endif // WITH_EDITOR
 
 #include "GameplayDebuggingComponent.h"
+#include "GAmeplayDebuggingControllerComponent.h"
+#include "GameplayDebuggingReplicator.h"
+
+DEFINE_LOG_CATEGORY(LogGameplayDebugger);
+
+FGameplayDebuggerShapeElement FGameplayDebuggerShapeElement::MakePoint(const FVector& Location, const float Radius, const FColor& Color, const FString& Description)
+{
+	FGameplayDebuggerShapeElement NewElement;
+	NewElement.Points.Add(Location);
+	NewElement.ThicknesOrRadius = Radius;
+	NewElement.Color = Color;
+	NewElement.Description = Description;
+	NewElement.Type = EGameplayDebuggerShapeElement::SinglePoint;
+
+	return NewElement;
+}
+
+FGameplayDebuggerShapeElement FGameplayDebuggerShapeElement::MakeSegment(const FVector& StartLocation, const FVector& EndLocation, const float Thickness, const FColor& Color, const FString& Description)
+{
+	FGameplayDebuggerShapeElement NewElement;
+	NewElement.Points.Add(StartLocation);
+	NewElement.Points.Add(EndLocation);
+	NewElement.ThicknesOrRadius = Thickness;
+	NewElement.Color = Color;
+	NewElement.Description = Description;
+	NewElement.Type = EGameplayDebuggerShapeElement::Segment;
+
+	return NewElement;
+}
+
+FGameplayDebuggerShapeElement FGameplayDebuggerShapeElement::MakeBox(const FVector& Center, const FVector& Extent, const FColor& Color, const FString& Description)
+{
+	FGameplayDebuggerShapeElement NewElement;
+	NewElement.Points.Add(Center);
+	NewElement.Points.Add(Extent);
+	NewElement.Color = Color;
+	NewElement.Description = Description;
+	NewElement.Type = EGameplayDebuggerShapeElement::Box;
+
+	return NewElement;
+}
+
+FGameplayDebuggerShapeElement FGameplayDebuggerShapeElement::MakeCone(const FVector& Location, const FVector& Direction, const float Length, const FColor& Color, const FString& Description)
+{
+	FGameplayDebuggerShapeElement NewElement;
+	NewElement.Points.Add(Location);
+	NewElement.Points.Add(Direction);
+	NewElement.ThicknesOrRadius = Length;
+	NewElement.Color = Color;
+	NewElement.Description = Description;
+	NewElement.Type = EGameplayDebuggerShapeElement::Cone;
+
+	return NewElement;
+}
+
+FGameplayDebuggerShapeElement FGameplayDebuggerShapeElement::MakeCylinder(const FVector& Center, const float Radius, const float HalfHeight, const FColor& Color, const FString& Description)
+{
+	FGameplayDebuggerShapeElement NewElement;
+	NewElement.Points.Add(Center - FVector(0, 0, HalfHeight));
+	NewElement.Points.Add(Center + FVector(0, 0, HalfHeight));
+	NewElement.ThicknesOrRadius = Radius;
+	NewElement.Color = Color;
+	NewElement.Description = Description;
+	NewElement.Type = EGameplayDebuggerShapeElement::Cylinder;
+
+	return NewElement;
+}
+
+FGameplayDebuggerShapeElement FGameplayDebuggerShapeElement::MakeCapsule(const FVector& Center, const float Radius, const float HalfHeight, const FColor& Color, const FString& Description)
+{
+	FGameplayDebuggerShapeElement NewElement;
+	NewElement.Points.Add(Center - FVector(0, 0, HalfHeight));
+	NewElement.Points.Add(Center + FVector(0, 0, HalfHeight));
+	NewElement.ThicknesOrRadius = Radius;
+	NewElement.Color = Color;
+	NewElement.Description = Description;
+	NewElement.Type = EGameplayDebuggerShapeElement::Capsule;
+
+	return NewElement;
+}
+
+FGameplayDebuggerShapeElement FGameplayDebuggerShapeElement::MakePolygon(const TArray<FVector>& Verts, const FColor& Color, const FString& Description)
+{
+	FGameplayDebuggerShapeElement NewElement;
+	NewElement.Points = Verts;
+	NewElement.Color = Color;
+	NewElement.Description = Description;
+	NewElement.Type = EGameplayDebuggerShapeElement::Polygon;
+	return NewElement;
+}
+
 
 //----------------------------------------------------------------------//
 // Composite Scene proxy
@@ -241,7 +336,7 @@ void UGameplayDebuggingComponent::GetLifetimeReplicatedProps( TArray< FLifetimeP
 	DOREPLIFETIME(UGameplayDebuggingComponent, DistanceFromPlayer);
 	DOREPLIFETIME(UGameplayDebuggingComponent, DistanceFromSensor);
 	DOREPLIFETIME(UGameplayDebuggingComponent, SensingComponentLocation);
-	DOREPLIFETIME(UGameplayDebuggingComponent, PerceptionShapeElements);
+//	DOREPLIFETIME(UGameplayDebuggingComponent, PerceptionShapeElements);
 
 #endif //!(UE_BUILD_SHIPPING || UE_BUILD_TEST)
 }
@@ -565,18 +660,10 @@ void UGameplayDebuggingComponent::CollectBasicBehaviorData(APawn* MyPawn)
 void UGameplayDebuggingComponent::CollectBasicAbilityData(APawn* MyPawn)
 {
 #if !(UE_BUILD_SHIPPING || UE_BUILD_TEST)
-	if (IGameplayAbilitiesModule::IsAvailable())
-	{
-		bool bUsingAbilities;
-		IGameplayAbilitiesModule& AbilitiesModule = FModuleManager::LoadModuleChecked<IGameplayAbilitiesModule>("GameplayAbilities");
-		AbilitiesModule.GetActiveAbilitiesDebugDataForActor(MyPawn, AbilityInfo, bUsingAbilities);
-		bIsUsingAbilities = bUsingAbilities;
-	}
-	else
-	{
-		bIsUsingAbilities = false;
-		AbilityInfo = TEXT("None");
-	}
+// DEPRECATED and not supported anymore, use new category based GameplayDebugger instead
+//		IGameplayAbilitiesModule& AbilitiesModule = FModuleManager::LoadModuleChecked<IGameplayAbilitiesModule>("GameplayAbilities");
+//		AbilitiesModule.GetActiveAbilitiesDebugDataForActor(MyPawn, AbilityInfo, bUsingAbilities);
+//		bIsUsingAbilities = bUsingAbilities;
 #endif //!(UE_BUILD_SHIPPING || UE_BUILD_TEST)
 }
 
@@ -1142,6 +1229,172 @@ ARecastNavMesh* UGameplayDebuggingComponent::GetNavData()
 }
 #endif
 
+static const FColor NavMeshRenderColor_Recast_TriangleEdges(255, 255, 255);
+static const FColor NavMeshRenderColor_Recast_TileEdges(16, 16, 16, 32);
+static const FColor NavMeshRenderColor_Recast_NavMeshEdges(32, 63, 0, 220);
+static const FColor NavMeshRenderColor_RecastMesh(140, 255, 0, 164);
+static const FColor NavMeshRenderColor_TileBounds(255, 255, 64, 255);
+static const FColor NavMeshRenderColor_PathCollidingGeom(255, 255, 255, 40);
+static const FColor NavMeshRenderColor_RecastTileBeingRebuilt(255, 0, 0, 64);
+static const FColor NavMeshRenderColor_OffMeshConnectionInvalid(64, 64, 64);
+static const float DefaultEdges_LineThickness = 0.0f;
+static const float PolyEdges_LineThickness = 1.5f;
+static const float NavMeshEdges_LineThickness = 3.5f;
+static const float LinkLines_LineThickness = 2.0f;
+static const float ClusterLinkLines_LineThickness = 2.0f;
+
+namespace FNavMeshRenderingHelpers_DEPRECATEDSUPPORT
+{
+	bool LineInView(const FVector& Start, const FVector& End, const FSceneView* View, bool bUseDistanceCheck)
+	{
+		if (FVector::DistSquared(Start, View->ViewMatrices.ViewOrigin) > ARecastNavMesh::GetDrawDistanceSq() ||
+			FVector::DistSquared(End, View->ViewMatrices.ViewOrigin) > ARecastNavMesh::GetDrawDistanceSq())
+		{
+			return false;
+		}
+
+		for (int32 PlaneIdx = 0; PlaneIdx < View->ViewFrustum.Planes.Num(); ++PlaneIdx)
+		{
+			const FPlane& CurPlane = View->ViewFrustum.Planes[PlaneIdx];
+			if (CurPlane.PlaneDot(Start) > 0.f && CurPlane.PlaneDot(End) > 0.f)
+			{
+				return false;
+			}
+		}
+
+		return true;
+	}
+
+	bool LineInCorrectDistance(const FVector& Start, const FVector& End, const FSceneView* View, float CorrectDistance = -1)
+	{
+		const float MaxDistanceSq = (CorrectDistance > 0) ? FMath::Square(CorrectDistance) : ARecastNavMesh::GetDrawDistanceSq();
+		return FVector::DistSquared(Start, View->ViewMatrices.ViewOrigin) < MaxDistanceSq &&
+			FVector::DistSquared(End, View->ViewMatrices.ViewOrigin) < MaxDistanceSq;
+	}
+
+	FVector EvalArc(const FVector& Org, const FVector& Dir, const float h, const float u)
+	{
+		FVector Pt = Org + Dir * u;
+		Pt.Z += h * (1 - (u * 2 - 1)*(u * 2 - 1));
+
+		return Pt;
+	}
+
+	void CacheArc(TArray<FDebugRenderSceneProxy::FDebugLine>& DebugLines, const FVector& Start, const FVector& End, const float Height, const uint32 Segments, const FLinearColor& Color, float LineThickness = 0)
+	{
+		if (Segments == 0)
+		{
+			return;
+		}
+
+		const float ArcPtsScale = 1.0f / (float)Segments;
+		const FVector Dir = End - Start;
+		const float Length = Dir.Size();
+
+		FVector Prev = Start;
+		for (uint32 i = 1; i <= Segments; ++i)
+		{
+			const float u = i * ArcPtsScale;
+			const FVector Pt = EvalArc(Start, Dir, Length*Height, u);
+
+			DebugLines.Add(FDebugRenderSceneProxy::FDebugLine(Prev, Pt, Color.ToFColor(true)));
+			Prev = Pt;
+		}
+	}
+
+	void CacheArrowHead(TArray<FDebugRenderSceneProxy::FDebugLine>& DebugLines, const FVector& Tip, const FVector& Origin, const float Size, const FLinearColor& Color, float LineThickness = 0)
+	{
+		FVector Ax, Ay, Az(0, 1, 0);
+		Ay = Origin - Tip;
+		Ay.Normalize();
+		Ax = FVector::CrossProduct(Az, Ay);
+
+		FHitProxyId HitProxyId;
+		DebugLines.Add(FDebugRenderSceneProxy::FDebugLine(Tip, FVector(Tip.X + Ay.X*Size + Ax.X*Size / 3, Tip.Y + Ay.Y*Size + Ax.Y*Size / 3, Tip.Z + Ay.Z*Size + Ax.Z*Size / 3), Color.ToFColor(true)));
+		DebugLines.Add(FDebugRenderSceneProxy::FDebugLine(Tip, FVector(Tip.X + Ay.X*Size - Ax.X*Size / 3, Tip.Y + Ay.Y*Size - Ax.Y*Size / 3, Tip.Z + Ay.Z*Size - Ax.Z*Size / 3), Color.ToFColor(true)));
+	}
+
+	void DrawWireCylinder(TArray<FDebugRenderSceneProxy::FDebugLine>& DebugLines, const FVector& Base, const FVector& X, const FVector& Y, const FVector& Z, FColor Color, float Radius, float HalfHeight, int32 NumSides, uint8 DepthPriority, float LineThickness = 0)
+	{
+		const float	AngleDelta = 2.0f * PI / NumSides;
+		FVector	LastVertex = Base + X * Radius;
+
+		FHitProxyId HitProxyId;
+		for (int32 SideIndex = 0; SideIndex < NumSides; SideIndex++)
+		{
+			const FVector Vertex = Base + (X * FMath::Cos(AngleDelta * (SideIndex + 1)) + Y * FMath::Sin(AngleDelta * (SideIndex + 1))) * Radius;
+
+			DebugLines.Add(FDebugRenderSceneProxy::FDebugLine(LastVertex - Z * HalfHeight, Vertex - Z * HalfHeight, Color));
+			DebugLines.Add(FDebugRenderSceneProxy::FDebugLine(LastVertex + Z * HalfHeight, Vertex + Z * HalfHeight, Color));
+			DebugLines.Add(FDebugRenderSceneProxy::FDebugLine(LastVertex - Z * HalfHeight, LastVertex + Z * HalfHeight, Color));
+
+			LastVertex = Vertex;
+		}
+	}
+
+	inline uint8 GetBit(int32 v, uint8 bit)
+	{
+		return (v & (1 << bit)) >> bit;
+	}
+
+	FColor GetClusterColor(int32 Idx)
+	{
+		uint8 r = 1 + GetBit(Idx, 1) + GetBit(Idx, 3) * 2;
+		uint8 g = 1 + GetBit(Idx, 2) + GetBit(Idx, 4) * 2;
+		uint8 b = 1 + GetBit(Idx, 0) + GetBit(Idx, 5) * 2;
+		return FColor(r * 63, g * 63, b * 63, 164);
+	}
+
+	FColor DarkenColor(const FColor& Base)
+	{
+		const uint32 Col = Base.DWColor();
+		return FColor(((Col >> 1) & 0x007f7f7f) | (Col & 0xff000000));
+	}
+
+	void AddVertex(FNavMeshSceneProxyData::FDebugMeshData& MeshData, const FVector& Pos, const FColor Color = FColor::White)
+	{
+		const int32 VertexIndex = MeshData.Vertices.Num();
+		FDynamicMeshVertex* Vertex = new(MeshData.Vertices) FDynamicMeshVertex;
+		Vertex->Position = Pos;
+		Vertex->TextureCoordinate = FVector2D::ZeroVector;
+		Vertex->TangentX = FVector(1.0f, 0.0f, 0.0f);
+		Vertex->TangentZ = FVector(0.0f, 1.0f, 0.0f);
+		// store the sign of the determinant in TangentZ.W (-1=0,+1=255)
+		Vertex->TangentZ.Vector.W = 255;
+		Vertex->Color = Color;
+	}
+
+	void AddTriangle(FNavMeshSceneProxyData::FDebugMeshData& MeshData, int32 V0, int32 V1, int32 V2)
+	{
+		MeshData.Indices.Add(V0);
+		MeshData.Indices.Add(V1);
+		MeshData.Indices.Add(V2);
+	}
+
+	FVector Recast2UnrealPoint(const float* RecastPoint)
+	{
+		return FVector(-RecastPoint[0], -RecastPoint[2], RecastPoint[1]);
+	}
+
+	void AddRecastGeometry(TArray<FVector>& OutVertexBuffer, TArray<int32>& OutIndexBuffer, const float* Coords, int32 NumVerts, const int32* Faces, int32 NumFaces)
+	{
+		const int32 VertIndexBase = OutVertexBuffer.Num();
+		for (int32 VertIdx = 0; VertIdx < NumVerts * 3; VertIdx += 3)
+		{
+			OutVertexBuffer.Add(Recast2UnrealPoint(&Coords[VertIdx]));
+		}
+
+		const int32 FirstNewFaceVertexIndex = OutIndexBuffer.Num();
+		const uint32 NumIndices = NumFaces * 3;
+		OutIndexBuffer.AddUninitialized(NumIndices);
+		for (uint32 Index = 0; Index < NumIndices; ++Index)
+		{
+			OutIndexBuffer[FirstNewFaceVertexIndex + Index] = VertIndexBase + Faces[Index];
+		}
+	}
+}
+
+
 void UGameplayDebuggingComponent::ServerCollectNavmeshData_Implementation(FVector_NetQuantize10 TargetLocation)
 {
 #if WITH_RECAST
@@ -1227,7 +1480,7 @@ void UGameplayDebuggingComponent::ServerCollectNavmeshData_Implementation(FVecto
 			NavMeshDebug::FOffMeshLink Link;
 			Link.Left = SrcLink.Left - TileData.Location;
 			Link.Right = SrcLink.Right - TileData.Location;
-			Link.Color = ((SrcLink.Direction && SrcLink.ValidEnds) || (SrcLink.ValidEnds & FRecastDebugGeometry::OMLE_Left)) ? DarkenColor(NavMeshColors[SrcLink.AreaID]) : NavMeshRenderColor_OffMeshConnectionInvalid;
+			Link.Color = ((SrcLink.Direction && SrcLink.ValidEnds) || (SrcLink.ValidEnds & FRecastDebugGeometry::OMLE_Left)) ? NavMeshColors[SrcLink.AreaID] : NavMeshRenderColor_OffMeshConnectionInvalid;
 			Link.PackedFlags.Radius = (int8)SrcLink.Radius;
 			Link.PackedFlags.Direction = SrcLink.Direction;
 			Link.PackedFlags.ValidEnds = SrcLink.ValidEnds;
@@ -1293,8 +1546,7 @@ void UGameplayDebuggingComponent::PrepareNavMeshData(struct FNavMeshSceneProxyDa
 		}
 
 		// read serialized values
-		CurrentData->bEnableDrawing = (UncompressedBuffer.Num() > 0);
-		if (!CurrentData->bEnableDrawing)
+		if (UncompressedBuffer.Num() == 0)
 		{
 			return;
 		}
@@ -1339,12 +1591,12 @@ void UGameplayDebuggingComponent::PrepareNavMeshData(struct FNavMeshSceneProxyDa
 					CurrentData->TileEdgeLines.Add(FDebugRenderSceneProxy::FDebugLine(V1 + CurrentData->NavMeshDrawOffset, V2 + CurrentData->NavMeshDrawOffset, NavMeshRenderColor_Recast_TileEdges));
 					CurrentData->TileEdgeLines.Add(FDebugRenderSceneProxy::FDebugLine(V2 + CurrentData->NavMeshDrawOffset, V0 + CurrentData->NavMeshDrawOffset, NavMeshRenderColor_Recast_TileEdges));
 
-					AddTriangleHelper(DebugMeshData, Index0 + IndexesOffset, Index1 + IndexesOffset, Index2 + IndexesOffset);
+					FNavMeshRenderingHelpers_DEPRECATEDSUPPORT::AddTriangle(DebugMeshData, Index0 + IndexesOffset, Index1 + IndexesOffset, Index2 + IndexesOffset);
 				}
 
 				for (int32 iVert = 0; iVert < Verts.Num(); iVert++)
 				{
-					AddVertexHelper(DebugMeshData, Verts[iVert] + CurrentData->NavMeshDrawOffset);
+					FNavMeshRenderingHelpers_DEPRECATEDSUPPORT::AddVertex(DebugMeshData, Verts[iVert] + CurrentData->NavMeshDrawOffset);
 				}
 				CurrentData->MeshBuilders.Add(DebugMeshData);
 				IndexesOffset += Verts.Num();
@@ -1358,13 +1610,13 @@ void UGameplayDebuggingComponent::PrepareNavMeshData(struct FNavMeshSceneProxyDa
 				const FVector V1 = SrcLink.Right.ToVector() + OffsetLocation + CurrentData->NavMeshDrawOffset;
 				const FColor LinkColor = SrcLink.Color;
 
-				CacheArc(CurrentData->NavLinkLines, V0, V1, 0.4f, 4, LinkColor, LinkLines_LineThickness);
+				FNavMeshRenderingHelpers_DEPRECATEDSUPPORT::CacheArc(CurrentData->NavLinkLines, V0, V1, 0.4f, 4, LinkColor, LinkLines_LineThickness);
 
 				const FVector VOffset(0, 0, FVector::Dist(V0, V1) * 1.333f);
-				CacheArrowHead(CurrentData->NavLinkLines, V1, V0+VOffset, 30.f, LinkColor, LinkLines_LineThickness);
+				FNavMeshRenderingHelpers_DEPRECATEDSUPPORT::CacheArrowHead(CurrentData->NavLinkLines, V1, V0+VOffset, 30.f, LinkColor, LinkLines_LineThickness);
 				if (SrcLink.PackedFlags.Direction)
 				{
-					CacheArrowHead(CurrentData->NavLinkLines, V0, V1+VOffset, 30.f, LinkColor, LinkLines_LineThickness);
+					FNavMeshRenderingHelpers_DEPRECATEDSUPPORT::CacheArrowHead(CurrentData->NavLinkLines, V0, V1+VOffset, 30.f, LinkColor, LinkLines_LineThickness);
 				}
 
 				// if the connection as a whole is valid check if there are any of ends is invalid
@@ -1373,12 +1625,12 @@ void UGameplayDebuggingComponent::PrepareNavMeshData(struct FNavMeshSceneProxyDa
 					if (SrcLink.PackedFlags.Direction && (SrcLink.PackedFlags.ValidEnds & FRecastDebugGeometry::OMLE_Left) == 0)
 					{
 						// left end invalid - mark it
-						DrawWireCylinder(CurrentData->NavLinkLines, V0, FVector(1, 0, 0), FVector(0, 1, 0), FVector(0, 0, 1), NavMeshRenderColor_OffMeshConnectionInvalid, SrcLink.PackedFlags.Radius, 30 /*NavMesh->AgentMaxStepHeight*/, 16, 0, DefaultEdges_LineThickness);
+						FNavMeshRenderingHelpers_DEPRECATEDSUPPORT::DrawWireCylinder(CurrentData->NavLinkLines, V0, FVector(1, 0, 0), FVector(0, 1, 0), FVector(0, 0, 1), NavMeshRenderColor_OffMeshConnectionInvalid, SrcLink.PackedFlags.Radius, 30 /*NavMesh->AgentMaxStepHeight*/, 16, 0, DefaultEdges_LineThickness);
 					}
 
 					if ((SrcLink.PackedFlags.ValidEnds & FRecastDebugGeometry::OMLE_Right) == 0)
 					{
-						DrawWireCylinder(CurrentData->NavLinkLines, V1, FVector(1, 0, 0), FVector(0, 1, 0), FVector(0, 0, 1), NavMeshRenderColor_OffMeshConnectionInvalid, SrcLink.PackedFlags.Radius, 30 /*NavMesh->AgentMaxStepHeight*/, 16, 0, DefaultEdges_LineThickness);
+						FNavMeshRenderingHelpers_DEPRECATEDSUPPORT::DrawWireCylinder(CurrentData->NavLinkLines, V1, FVector(1, 0, 0), FVector(0, 1, 0), FVector(0, 0, 1), NavMeshRenderColor_OffMeshConnectionInvalid, SrcLink.PackedFlags.Radius, 30 /*NavMesh->AgentMaxStepHeight*/, 16, 0, DefaultEdges_LineThickness);
 					}
 				}
 			}
@@ -1437,12 +1689,11 @@ FPrimitiveSceneProxy* UGameplayDebuggingComponent::CreateSceneProxy()
 		FNavMeshSceneProxyData NewNavmeshRenderData;
 		NewNavmeshRenderData.Reset();
 		NewNavmeshRenderData.bNeedsNewData = false;
-		NewNavmeshRenderData.bEnableDrawing = false;
 		PrepareNavMeshData(&NewNavmeshRenderData);
 
 		NavMeshBounds = NewNavmeshRenderData.Bounds.GetCenter().ContainsNaN() || NewNavmeshRenderData.Bounds.GetExtent().ContainsNaN() ? FBox(FVector(-HALF_WORLD_MAX1, -HALF_WORLD_MAX1, -HALF_WORLD_MAX1), FVector(HALF_WORLD_MAX1, HALF_WORLD_MAX1, HALF_WORLD_MAX1)) : NewNavmeshRenderData.Bounds;
 		CompositeProxy = CompositeProxy ? CompositeProxy : (new FDebugRenderSceneCompositeProxy(this));
-		CompositeProxy->AddChild(new FRecastRenderingSceneProxy(this, &NewNavmeshRenderData, true));
+		CompositeProxy->AddChild(new FNavMeshSceneProxy(this, &NewNavmeshRenderData, true));
 	}
 #endif
 
@@ -1579,6 +1830,7 @@ void UGameplayDebuggingComponent::CollectPerceptionData()
 		return;
 	}
 
+#if 0
 	APawn* MyPawn = Cast<APawn>(GetSelectedActor());
 	if (MyPawn)
 	{
@@ -1616,5 +1868,7 @@ void UGameplayDebuggingComponent::CollectPerceptionData()
 			}
 		}
 	}
+
+#endif
 #endif
 }
diff --git a/Engine/Source/Developer/GameplayDebugger/Private/GameplayDebuggingControllerComponent.cpp b/Engine/Source/Developer/GameplayDebugger/Private/GameplayDebuggingControllerComponent.cpp
--- a/Engine/Source/Developer/GameplayDebugger/Private/GameplayDebuggingControllerComponent.cpp
+++ b/Engine/Source/Developer/GameplayDebugger/Private/GameplayDebuggingControllerComponent.cpp
@@ -1,12 +1,22 @@
 // Copyright 1998-2016 Epic Games, Inc. All Rights Reserved.
-#include "GameplayDebuggerPrivate.h"
+
+//////////////////////////////////////////////////////////////////////////
+// THIS CLASS IS NOW DEPRECATED AND WILL BE REMOVED IN NEXT VERSION
+// Please check GameplayDebugger.h for details.
+
+#include "GameplayDebuggerPrivatePCH.h"
 #include "Net/UnrealNetwork.h"
 #include "GameFramework/HUD.h"
+#include "GameFramework/Pawn.h"
 #include "GameFramework/PlayerInput.h"
+#include "GameFramework/WorldSettings.h"
 #include "BehaviorTreeDelegates.h"
 #include "TimerManager.h"
 #include "VisualLogger/VisualLogger.h"
 #include "GameplayDebuggerSettings.h"
+#include "GameplayDebuggingControllerComponent.h"
+#include "GameplayDebuggingComponent.h"
+#include "GameplayDebuggingReplicator.h"
 #include "Engine/Player.h"
 
 DEFINE_LOG_CATEGORY_STATIC(LogGameplayDebugging, Log, All);
diff --git a/Engine/Source/Developer/GameplayDebugger/Private/GameplayDebuggingHUDComponent.cpp b/Engine/Source/Developer/GameplayDebugger/Private/GameplayDebuggingHUDComponent.cpp
--- a/Engine/Source/Developer/GameplayDebugger/Private/GameplayDebuggingHUDComponent.cpp
+++ b/Engine/Source/Developer/GameplayDebugger/Private/GameplayDebuggingHUDComponent.cpp
@@ -1,15 +1,18 @@
 // Copyright 1998-2016 Epic Games, Inc. All Rights Reserved.
 
-/*=============================================================================
-	HUD.cpp: Heads up Display related functionality
-=============================================================================*/
-#include "GameplayDebuggerPrivate.h"
+//////////////////////////////////////////////////////////////////////////
+// THIS CLASS IS NOW DEPRECATED AND WILL BE REMOVED IN NEXT VERSION
+// Please check GameplayDebugger.h for details.
+
+#include "GameplayDebuggerPrivatePCH.h"
 #include "GameplayDebuggerSettings.h"
 #include "Net/UnrealNetwork.h"
 #include "GameplayDebuggingComponent.h"
 #include "GameplayDebuggingHUDComponent.h"
 #include "GameplayDebuggingControllerComponent.h"
+#include "GameplayDebuggingReplicator.h"
 #include "CanvasItem.h"
+#include "Engine/Canvas.h"
 #include "AI/Navigation/NavigationSystem.h"
 
 #include "AITypes.h"
@@ -17,13 +20,13 @@
 #include "GenericTeamAgentInterface.h"
 #include "AIController.h"
 
-
 #include "EnvironmentQuery/EnvQueryTypes.h"
 #include "Engine/Texture2D.h"
 #include "Regex.h"
 #include "DrawDebugHelpers.h"
 #include "TimerManager.h"					// Game play timers
 #include "RenderUtils.h"
+#include "EngineUtils.h"
 
 DEFINE_LOG_CATEGORY_STATIC(LogHUD, Log, All);
 
diff --git a/Engine/Source/Developer/GameplayDebugger/Private/GameplayDebuggingReplicator.cpp b/Engine/Source/Developer/GameplayDebugger/Private/GameplayDebuggingReplicator.cpp
--- a/Engine/Source/Developer/GameplayDebugger/Private/GameplayDebuggingReplicator.cpp
+++ b/Engine/Source/Developer/GameplayDebugger/Private/GameplayDebuggingReplicator.cpp
@@ -1,12 +1,20 @@
 // Copyright 1998-2016 Epic Games, Inc. All Rights Reserved.
 
-#include "GameplayDebuggerPrivate.h"
+//////////////////////////////////////////////////////////////////////////
+// THIS CLASS IS NOW DEPRECATED AND WILL BE REMOVED IN NEXT VERSION
+// Please check GameplayDebugger.h for details.
+
+#include "GameplayDebuggerPrivatePCH.h"
 #include "Engine/GameInstance.h"
+#include "Engine/Canvas.h"
 #include "Debug/DebugDrawService.h"
 #include "GameFramework/HUD.h"
+#include "GameFramework/Pawn.h"
+#include "GameplayDebugger.h"
 #include "GameplayDebuggingComponent.h"
 #include "GameplayDebuggingHUDComponent.h"
 #include "GameplayDebuggingReplicator.h"
+#include "GameplayDebuggingControllerComponent.h"
 #include "BehaviorTreeDelegates.h"
 #if WITH_EDITOR
 #include "Editor/EditorEngine.h"
@@ -506,7 +514,7 @@ void AGameplayDebuggingReplicator::TickActor(float DeltaTime, enum ELevelTick Ti
 
 #if !(UE_BUILD_SHIPPING || UE_BUILD_TEST)
 	UWorld* World = GetWorld();
-	if (!IsGlobalInWorld() || !World || GetNetMode() == ENetMode::NM_Client || !IGameplayDebugger::IsAvailable())
+	if (!IsGlobalInWorld() || !World || GetNetMode() == ENetMode::NM_Client || !GameplayDebugger::IsAvailable())
 	{
 		// global level replicator don't have any local player and it's prepared to work only on servers
 		return;
@@ -526,7 +534,7 @@ void AGameplayDebuggingReplicator::TickActor(float DeltaTime, enum ELevelTick Ti
 			APlayerController* PC = *Iterator;
 			if (PC)
 			{
-				IGameplayDebugger& Debugger = IGameplayDebugger::Get();
+				GameplayDebugger& Debugger = GameplayDebugger::Get();
 				Debugger.CreateGameplayDebuggerForPlayerController(PC);
 			}
 		}
diff --git a/Engine/Source/Developer/GameplayDebugger/Public/GameplayDebugger.h b/Engine/Source/Developer/GameplayDebugger/Public/GameplayDebugger.h
--- a/Engine/Source/Developer/GameplayDebugger/Public/GameplayDebugger.h
+++ b/Engine/Source/Developer/GameplayDebugger/Public/GameplayDebugger.h
@@ -1,16 +1,71 @@
 // Copyright 1998-2016 Epic Games, Inc. All Rights Reserved.
 
+// GAMEPLAY DEBUGGER
+// 
+// This tool allows easy on screen debugging of gameplay data, supporting client-server replication.
+// Data is organized into named categories, which can be toggled during debugging.
+// 
+// To enable it, press Apostrophe key (UGameplayDebuggerConfig::ActivationKey)
+//
+// Category class:
+// - derives from FGameplayDebuggerCategory
+// - implements at least CollectData() and DrawData() functions
+// - requires WITH_GAMEPLAY_DEBUGGER define to compile (doesn't exist in shipping builds by default)
+// - needs to be registered and unregistered manually by owning module
+// - automatically replicate data added with FGameplayDebuggerCategory::AddTextLine, FGameplayDebuggerCategory::AddShape
+// - automatically replicate data structs initialized with FGameplayDebuggerCategory::SetDataPackReplication
+// - can define own input bindings (e.g. subcategories, etc)
+//
+// Extension class:
+// - derives from FGameplayDebuggerExtension
+// - needs to be registered and unregistered manually by owning module
+// - can define own input bindings
+// - basically it's a stateless, not replicated, not drawn category, ideal for making e.g. different actor selection mechanic
+//
+// 
+// Check FGameplayDebuggerCategory_BehaviorTree for implementation example.
+// Check AIModule/Private/AIModule.cpp for registration example.
+//
+//
+// Remember to define WITH_GAMEPLAY_DEBUGGER=1 when adding module to your project's Build.cs!
+// 
+// if (UEBuildConfiguration.bBuildDeveloperTools &&
+//     Target.Configuration != UnrealTargetConfiguration.Shipping &&
+//     Target.Configuration != UnrealTargetConfiguration.Test)
+// {
+//     PrivateDependencyModuleNames.Add("GameplayDebugger");
+//     Definitions.Add("WITH_GAMEPLAY_DEBUGGER=1");
+// }
+//
+//
+// -----------------------------------------------------------------------------------------------------
+// UPGRADE NOTES:
+// 
+// GameplayDebuggingComponent and GameplayDebuggingHUDComponent classes are no longer supported.
+// Instead of adding new project specific module, now you can create FGameplayDebuggerCategory classes
+// within project module. 
+// 
+// If you need to use old GameplayDebugger implementation, please set bEnableDeprecatedDebugger flag in GameplayDebugger.Build.cs
+// 
+
 #pragma once
 
 #include "ModuleManager.h"
 
-/**
- * The public interface to this module
- */
+enum class EGameplayDebuggerCategoryState : uint8
+{
+	EnabledInGameAndSimulate,
+	EnabledInGame,
+	EnabledInSimulate,
+	Disabled,
+};
+
 class IGameplayDebugger : public IModuleInterface
 {
 
 public:
+	DECLARE_DELEGATE_RetVal(TSharedRef<class FGameplayDebuggerCategory>, FOnGetCategory);
+	DECLARE_DELEGATE_RetVal(TSharedRef<class FGameplayDebuggerExtension>, FOnGetExtension);
 
 	/**
 	 * Singleton-like access to this module's interface.  This is just for convenience!
@@ -20,8 +75,7 @@ public:
 	 */
 	static inline IGameplayDebugger& Get()
 	{
-		static FName GamePlayDebuggerModuleName("GameplayDebugger");
-		return FModuleManager::LoadModuleChecked< IGameplayDebugger >(GamePlayDebuggerModuleName);
+		return FModuleManager::LoadModuleChecked< IGameplayDebugger >("GameplayDebugger");
 	}
 
 	/**
@@ -31,16 +85,38 @@ public:
 	 */
 	static inline bool IsAvailable()
 	{
-		static FName GamePlayDebuggerModuleName("GameplayDebugger");
-		return FModuleManager::Get().IsModuleLoaded( GamePlayDebuggerModuleName );
+		return FModuleManager::Get().IsModuleLoaded("GameplayDebugger");
 	}
 
-	// Each player controller that wants to use gameplay debugging must call this function
-	// (generally OnPostInitProperties, but not on a client) in order to create the actor that handles the debugging
-	// functionality in a network-replicated (if necessary) fashion.  NOTE: creates an AGameplayDebuggingReplicator
-	// in the same World as PlayerController.
-	virtual bool CreateGameplayDebuggerForPlayerController(APlayerController* PlayerController) = 0;
+	virtual void RegisterCategory(FName CategoryName, FOnGetCategory MakeInstanceDelegate, EGameplayDebuggerCategoryState CategoryState = EGameplayDebuggerCategoryState::Disabled, int32 SlotIdx = INDEX_NONE) = 0;
+	virtual void UnregisterCategory(FName CategoryName) = 0;
+	virtual void NotifyCategoriesChanged() = 0;
+	virtual void RegisterExtension(FName ExtensionName, FOnGetExtension MakeInstanceDelegate) = 0;
+	virtual void UnregisterExtension(FName ExtensionName) = 0;
+	virtual void NotifyExtensionsChanged() = 0;
 
+#if ENABLE_OLD_GAMEPLAY_DEBUGGER
+	virtual bool CreateGameplayDebuggerForPlayerController(APlayerController* PlayerController) = 0;
 	virtual bool IsGameplayDebuggerActiveForPlayerController(APlayerController* PlayerController) = 0;
+#else
+	DEPRECATED_FORGAME(4.12, "This function is now deprecated, please check GameplayDebugger.h for details.")
+	virtual bool CreateGameplayDebuggerForPlayerController(APlayerController* PlayerController) { return false; }
+	
+	DEPRECATED_FORGAME(4.12, "This function is now deprecated, please check GameplayDebugger.h for details.")
+	virtual bool IsGameplayDebuggerActiveForPlayerController(APlayerController* PlayerController) { return false; }
+#endif
 };
 
+#if ENABLE_OLD_GAMEPLAY_DEBUGGER
+class GameplayDebugger : public IGameplayDebugger
+{
+public:
+	static inline GameplayDebugger& Get() { return FModuleManager::LoadModuleChecked< GameplayDebugger >("GameplayDebugger"); }
+};
+#else
+class DEPRECATED_FORGAME(4.12, "This module interface is now deprecated, please check GameplayDebugger.h for details.") GameplayDebugger : public IGameplayDebugger
+{
+public:
+	static inline GameplayDebugger& Get() { return FModuleManager::LoadModuleChecked< GameplayDebugger >("GameplayDebugger"); }
+};
+#endif
diff --git a/Engine/Source/Developer/GameplayDebugger/Public/GameplayDebuggerAddonBase.h b/Engine/Source/Developer/GameplayDebugger/Public/GameplayDebuggerAddonBase.h
new file mode 100644
--- /dev/null
+++ b/Engine/Source/Developer/GameplayDebugger/Public/GameplayDebuggerAddonBase.h
@@ -0,0 +1,73 @@
+// Copyright 1998-2016 Epic Games, Inc. All Rights Reserved.
+
+#pragma once
+
+#include "GameplayDebuggerTypes.h"
+#include "GameplayDebuggerCategoryReplicator.h"
+
+class GAMEPLAYDEBUGGER_API FGameplayDebuggerAddonBase : public TSharedFromThis<FGameplayDebuggerAddonBase>
+{
+public:
+
+	virtual ~FGameplayDebuggerAddonBase() {}
+
+	int32 GetNumInputHandlers() const { return InputHandlers.Num(); }
+	FGameplayDebuggerInputHandler& GetInputHandler(int32 HandlerId) { return InputHandlers[HandlerId]; }
+	FString GetInputHandlerDescription(int32 HandlerId) const;
+
+	/** check if simulate in editor mode is active */
+	static bool IsSimulateInEditor();
+
+protected:
+
+	/** tries to find selected actor in local world */
+	AActor* FindLocalDebugActor() const;
+
+	/** returns replicator actor */
+	AGameplayDebuggerCategoryReplicator* GetReplicator() const;
+
+	/** creates new key binding handler */
+	template<class UserClass>
+	bool BindKeyPress(FName KeyName, UserClass* KeyHandlerObject, typename FGameplayDebuggerInputHandler::FHandler::TRawMethodDelegate< UserClass >::FMethodPtr KeyHandlerFunc)
+	{
+		FGameplayDebuggerInputHandler NewHandler;
+		NewHandler.KeyName = KeyName;
+		NewHandler.Delegate.BindRaw(KeyHandlerObject, KeyHandlerFunc);
+
+		if (NewHandler.IsValid())
+		{
+			InputHandlers.Add(NewHandler);
+			return true;
+		}
+
+		return false;
+	}
+
+	/** creates new key binding handler */
+	template<class UserClass>
+	bool BindKeyPress(FName KeyName, FGameplayDebuggerInputModifier KeyModifer, UserClass* KeyHandlerObject, typename FGameplayDebuggerInputHandler::FHandler::TRawMethodDelegate< UserClass >::FMethodPtr KeyHandlerFunc)
+	{
+		FGameplayDebuggerInputHandler NewHandler;
+		NewHandler.KeyName = KeyName;
+		NewHandler.Modifier = KeyModifer;
+		NewHandler.Delegate.BindRaw(KeyHandlerObject, KeyHandlerFunc);
+
+		if (NewHandler.IsValid())
+		{
+			InputHandlers.Add(NewHandler);
+			return true;
+		}
+
+		return false;
+	}
+
+private:
+
+	friend class FGameplayDebuggerAddonManager;
+
+	/** list of input handlers */
+	TArray<FGameplayDebuggerInputHandler> InputHandlers;
+
+	/** replicator actor */
+	TWeakObjectPtr<AGameplayDebuggerCategoryReplicator> RepOwner;
+};
diff --git a/Engine/Source/Developer/GameplayDebugger/Public/GameplayDebuggerCategoryReplicator.h b/Engine/Source/Developer/GameplayDebugger/Public/GameplayDebuggerCategoryReplicator.h
new file mode 100644
--- /dev/null
+++ b/Engine/Source/Developer/GameplayDebugger/Public/GameplayDebuggerCategoryReplicator.h
@@ -0,0 +1,151 @@
+// Copyright 1998-2016 Epic Games, Inc. All Rights Reserved.
+
+#pragma once
+
+#include "GameplayDebuggerTypes.h"
+#include "GameplayDebuggerCategoryReplicator.generated.h"
+
+class APlayerController;
+class UInputComponent;
+class FGameplayDebuggerCategory;
+class AGameplayDebuggerCategoryReplicator;
+class UGameplayDebuggerRenderingComponent;
+
+USTRUCT()
+struct FGameplayDebuggerNetPack
+{
+	GENERATED_USTRUCT_BODY()
+
+	AGameplayDebuggerCategoryReplicator* Owner;
+
+	FGameplayDebuggerNetPack() : Owner(nullptr) {}
+	bool NetDeltaSerialize(FNetDeltaSerializeInfo & DeltaParms);
+	void OnCategoriesChanged();
+
+private:
+	struct FCategoryData
+	{
+		TArray<FString> TextLines;
+		TArray<FGameplayDebuggerShape> Shapes;
+		TArray<FGameplayDebuggerDataPack::FHeader> DataPacks;
+		bool bIsEnabled;
+	};
+	TArray<FCategoryData> SavedData;
+};
+
+template<>
+struct TStructOpsTypeTraits<FGameplayDebuggerNetPack> : public TStructOpsTypeTraitsBase
+{
+	enum
+	{
+		WithNetDeltaSerializer = true,
+	};
+};
+
+USTRUCT()
+struct FGameplayDebuggerDataPacketRequest
+{
+	GENERATED_USTRUCT_BODY()
+
+	UPROPERTY()
+	int32 CategoryId;
+
+	UPROPERTY()
+	int32 DataPackId;
+
+	UPROPERTY()
+	int32 DataVersion;
+
+	UPROPERTY()
+	int32 DataOffset;
+};
+
+UCLASS(NotBlueprintable, NotBlueprintType, notplaceable, noteditinlinenew, hidedropdown, Transient)
+class GAMEPLAYDEBUGGER_API AGameplayDebuggerCategoryReplicator : public AActor
+{
+	GENERATED_UCLASS_BODY()
+
+	virtual class UNetConnection* GetNetConnection() const override;
+	virtual bool IsNetRelevantFor(const AActor* RealViewer, const AActor* ViewTarget, const FVector& SrcLocation) const override;
+	virtual void PostNetReceive() override;
+	virtual void BeginPlay() override;
+	virtual void Destroyed() override;
+	virtual void TickActor(float DeltaTime, enum ELevelTick TickType, FActorTickFunction& ThisTickFunction) override;
+
+	/** [AUTH] set new owner */
+	void SetReplicatorOwner(APlayerController* InOwnerPC);
+
+	/** [ALL] set replicator state */
+	void SetEnabled(bool bEnable);
+
+	/** [ALL] set category state */
+	void SetCategoryEnabled(int32 CategoryId, bool bEnable);
+
+	/** [ALL] set actor for debugging */
+	void SetDebugActor(AActor* Actor);
+
+	/** get current debug actor */
+	AActor* GetDebugActor() const { return DebugActor; }
+	
+	/** get player controller owning this replicator */
+	APlayerController* GetReplicationOwner() const { return OwnerPC; }
+
+	/** get replicator state */
+	bool IsEnabled() const { return bIsEnabled; }
+
+	/** get category state */
+	bool IsCategoryEnabled(int32 CategoryId) const;
+
+	/** get category count */
+	int32 GetNumCategories() const { return Categories.Num(); }
+
+	/** get category object */
+	TSharedRef<FGameplayDebuggerCategory> GetCategory(int32 CategoryId) const { return Categories[CategoryId]; }
+
+protected:
+
+	friend FGameplayDebuggerNetPack;
+
+	UPROPERTY(Replicated)
+	APlayerController* OwnerPC;
+
+	UPROPERTY(Replicated)
+	AActor*	DebugActor;
+
+	UPROPERTY(Replicated)
+	bool bIsEnabled;
+
+	UPROPERTY(Replicated)
+	FGameplayDebuggerNetPack ReplicatedData;
+
+	/** rendering component needs to attached to some actor, and this is as good as any */
+	UPROPERTY()
+	UGameplayDebuggerRenderingComponent* RenderingComp;
+
+	/** category objects */
+	TArray<TSharedRef<FGameplayDebuggerCategory> > Categories;
+
+	/** pending data pack replication requests */
+	TArray<FGameplayDebuggerDataPacketRequest> PendingReplicationRequests;
+
+	uint32 bHasAuthority : 1;
+	uint32 bIsLocal : 1;
+
+	/** notify about changes in known category set */
+	void OnCategoriesChanged();
+
+	UFUNCTION(Server, Reliable, WithValidation, meta = (CallInEditor = "true"))
+	void ServerSetEnabled(bool bEnable);
+
+	UFUNCTION(Server, Reliable, WithValidation, meta = (CallInEditor = "true"))
+	void ServerSetDebugActor(AActor* Actor);
+
+	UFUNCTION(Server, Reliable, WithValidation, meta = (CallInEditor = "true"))
+	void ServerSetCategoryEnabled(int32 CategoryId, bool bEnable);
+
+	UFUNCTION(Server, Reliable, WithValidation, meta = (CallInEditor = "true"))
+	void ServerRequestDataPackets(const TArray<FGameplayDebuggerDataPacketRequest>& RequestList);
+
+	/** [LOCAL] notify from CategoryData replication */
+	void OnReceivedDataPackPacket(int32 CategoryId, int32 DataPackId, const FGameplayDebuggerDataPack& DataPacket);
+};
diff --git a/Engine/Source/Developer/GameplayDebugger/Public/GameplayDebuggerConfig.h b/Engine/Source/Developer/GameplayDebugger/Public/GameplayDebuggerConfig.h
new file mode 100644
--- /dev/null
+++ b/Engine/Source/Developer/GameplayDebugger/Public/GameplayDebuggerConfig.h
@@ -0,0 +1,141 @@
+// Copyright 1998-2016 Epic Games, Inc. All Rights Reserved.
+
+#pragma once
+
+#include "GameplayDebuggerTypes.h"
+#include "GameplayDebuggerConfig.generated.h"
+
+USTRUCT()
+struct GAMEPLAYDEBUGGER_API FGameplayDebuggerInputConfig
+{
+	GENERATED_USTRUCT_BODY();
+
+	UPROPERTY(EditAnywhere, Category = Input)
+	uint32 bShift : 1;
+
+	UPROPERTY(EditAnywhere, Category = Input)
+	uint32 bCtrl : 1;
+
+	UPROPERTY(EditAnywhere, Category = Input)
+	uint32 bAlt : 1;
+
+	UPROPERTY(EditAnywhere, Category = Input)
+	uint32 bCmd : 1;
+
+	FGameplayDebuggerInputModifier CreateModifier() const
+	{
+		return FGameplayDebuggerInputModifier(bShift, bCtrl, bAlt, bCmd);
+	}
+};
+
+UCLASS(config = Engine, defaultconfig)
+class GAMEPLAYDEBUGGER_API UGameplayDebuggerConfig : public UObject
+{
+	GENERATED_UCLASS_BODY()
+
+	/** key used to activate visual debugger tool */
+	UPROPERTY(config, EditAnywhere, Category = Input)
+	FKey ActivationKey;
+
+	/** select next category row */
+	UPROPERTY(config, EditAnywhere, Category = Input)
+	FKey CategoryRowNextKey;
+
+	/** select previous category row */
+	UPROPERTY(config, EditAnywhere, Category = Input)
+	FKey CategoryRowPrevKey;
+
+	/** select category slot 0 */
+	UPROPERTY(config, EditAnywhere, Category = Input)
+	FKey CategorySlot0;
+
+	/** select category slot 1 */
+	UPROPERTY(config, EditAnywhere, Category = Input)
+	FKey CategorySlot1;
+
+	/** select category slot 2 */
+	UPROPERTY(config, EditAnywhere, Category = Input)
+	FKey CategorySlot2;
+
+	/** select category slot 3 */
+	UPROPERTY(config, EditAnywhere, Category = Input)
+	FKey CategorySlot3;
+
+	/** select category slot 4 */
+	UPROPERTY(config, EditAnywhere, Category = Input)
+	FKey CategorySlot4;
+
+	/** select category slot 5 */
+	UPROPERTY(config, EditAnywhere, Category = Input)
+	FKey CategorySlot5;
+
+	/** select category slot 6 */
+	UPROPERTY(config, EditAnywhere, Category = Input)
+	FKey CategorySlot6;
+
+	/** select category slot 7 */
+	UPROPERTY(config, EditAnywhere, Category = Input)
+	FKey CategorySlot7;
+
+	/** select category slot 8 */
+	UPROPERTY(config, EditAnywhere, Category = Input)
+	FKey CategorySlot8;
+
+	/** select category slot 9 */
+	UPROPERTY(config, EditAnywhere, Category = Input)
+	FKey CategorySlot9;
+
+	/** toggle game hud */
+	UPROPERTY(config, EditAnywhere, Category = Extensions)
+	FKey GameHUDKey;
+
+	UPROPERTY(config, EditAnywhere, Category = Extensions, meta = (EditCondition = "bUseGameHUDModifiers"))
+	FGameplayDebuggerInputConfig GameHUDModifiers;
+
+	/** toggle onscreen debug messages */
+	UPROPERTY(config, EditAnywhere, Category = Extensions)
+	FKey DebugMessagesKey;
+
+	UPROPERTY(config, EditAnywhere, Category = Extensions, meta = (EditCondition = "bUseDebugMessagesModifiers"))
+	FGameplayDebuggerInputConfig DebugMessagesModifiers;
+
+	/** toggle spectator mode */
+	UPROPERTY(config, EditAnywhere, Category = Extensions)
+	FKey SpectatorKey;
+
+	UPROPERTY(config, EditAnywhere, Category = Extensions, meta = (EditCondition = "bUseSpectatorModifiers"))
+	FGameplayDebuggerInputConfig SpectatorModifiers;
+
+	/** if set, game hud extension will be automatically added */
+	UPROPERTY(config, EditAnywhere, Category = Extensions)
+	uint32 bEnableExtension_GameHUD : 1;
+
+	/** if set, spectator mode extension will be automatically added */
+	UPROPERTY(config, EditAnywhere, Category = Extensions)
+	uint32 bEnableExtension_Spectator : 1;
+
+	UPROPERTY()
+	uint32 bUseGameHUDModifiers : 1;
+	
+	UPROPERTY()
+	uint32 bUseDebugMessagesModifiers : 1;
+	
+	UPROPERTY()
+	uint32 bUseSpectatorModifiers : 1;
+
+	/** additional canvas padding: left */
+	UPROPERTY(config, EditAnywhere, Category = Display)
+	float DebugCanvasPaddingLeft;
+
+	/** additional canvas padding: right */
+	UPROPERTY(config, EditAnywhere, Category = Display)
+	float DebugCanvasPaddingRight;
+
+	/** additional canvas padding: top */
+	UPROPERTY(config, EditAnywhere, Category = Display)
+	float DebugCanvasPaddingTop;
+
+	/** additional canvas padding: bottom */
+	UPROPERTY(config, EditAnywhere, Category = Display)
+	float DebugCanvasPaddingBottom;
+};
diff --git a/Engine/Source/Developer/GameplayDebugger/Public/GameplayDebuggerExtension.h b/Engine/Source/Developer/GameplayDebugger/Public/GameplayDebuggerExtension.h
new file mode 100644
--- /dev/null
+++ b/Engine/Source/Developer/GameplayDebugger/Public/GameplayDebuggerExtension.h
@@ -0,0 +1,37 @@
+// Copyright 1998-2016 Epic Games, Inc. All Rights Reserved.
+
+// GAMEPLAY DEBUGGER EXTENSION
+// 
+// Extensions allows creating additional key bindings for gameplay debugger for custom functionality.
+// For example, you can use them to add another way of selecting actor to Debug.
+//
+// It should be compiled and used only when module is included, so every extension class
+// needs be placed in #if WITH_GAMEPLAY_DEBUGGER block.
+// 
+// Extensions needs to be manually registered and unregistered with GameplayDebugger.
+// It's best to do it in owning module's Startup / Shutdown, similar to detail view customizations.
+
+#pragma once
+
+#include "GameplayDebuggerAddonBase.h"
+
+class GAMEPLAYDEBUGGER_API FGameplayDebuggerExtension : public FGameplayDebuggerAddonBase
+{
+public:
+
+	virtual ~FGameplayDebuggerExtension() {}
+
+	/** called when extension is added to debugger or debugger is enabled */
+	virtual void OnActivated();
+
+	/** called when extension is removed from debugger or debugger is disabled */
+	virtual void OnDeactivated();
+
+	/** description for gameplay debugger's header row, newline character is ignored */
+	virtual FString GetDescription() const;
+
+protected:
+
+	/** get player controller owning gameplay debugger tool */
+	APlayerController* GetPlayerController() const;
+};
diff --git a/Engine/Source/Developer/GameplayDebugger/Public/GameplayDebuggerPlayerManager.h b/Engine/Source/Developer/GameplayDebugger/Public/GameplayDebuggerPlayerManager.h
new file mode 100644
--- /dev/null
+++ b/Engine/Source/Developer/GameplayDebugger/Public/GameplayDebuggerPlayerManager.h
@@ -0,0 +1,51 @@
+// Copyright 1998-2016 Epic Games, Inc. All Rights Reserved.
+
+#pragma once
+
+#include "GameFramework/PlayerController.h"
+#include "GameplayDebuggerPlayerManager.generated.h"
+
+class AGameplayDebuggerCategoryReplicator;
+class UGameplayDebuggerLocalController;
+
+USTRUCT()
+struct FGameplayDebuggerPlayerData
+{
+	GENERATED_USTRUCT_BODY()
+
+	UPROPERTY()
+	UGameplayDebuggerLocalController* Controller;
+
+	UPROPERTY()
+	UInputComponent* InputComponent;
+
+	UPROPERTY()
+	AGameplayDebuggerCategoryReplicator* Replicator;
+};
+
+UCLASS(NotBlueprintable, NotBlueprintType, notplaceable, noteditinlinenew, hidedropdown, Transient)
+class AGameplayDebuggerPlayerManager : public AActor
+{
+	GENERATED_UCLASS_BODY()
+
+	virtual void TickActor(float DeltaTime, enum ELevelTick TickType, FActorTickFunction& ThisTickFunction) override;
+	virtual void BeginPlay() override;
+	virtual void EndPlay(const EEndPlayReason::Type Reason) override;
+		
+	void UpdateAuthReplicators();
+	void RegisterReplicator(AGameplayDebuggerCategoryReplicator& Replicator);
+	void RefreshInputBindings(AGameplayDebuggerCategoryReplicator& Replicator);
+
+	AGameplayDebuggerCategoryReplicator* GetReplicator(const APlayerController& OwnerPC) const;
+	UInputComponent* GetInputComponent(const APlayerController& OwnerPC) const;
+	
+	static AGameplayDebuggerPlayerManager& GetCurrent(UWorld* World);
+
+protected:
+
+	UPROPERTY()
+	TArray<FGameplayDebuggerPlayerData> PlayerData;
+
+	uint32 bHasAuthority : 1;
+	uint32 bIsLocal : 1;
+};
diff --git a/Engine/Source/Developer/GameplayDebugger/Public/GameplayDebuggerRenderingComponent.h b/Engine/Source/Developer/GameplayDebugger/Public/GameplayDebuggerRenderingComponent.h
new file mode 100644
--- /dev/null
+++ b/Engine/Source/Developer/GameplayDebugger/Public/GameplayDebuggerRenderingComponent.h
@@ -0,0 +1,18 @@
+// Copyright 1998-2016 Epic Games, Inc. All Rights Reserved.
+
+#pragma once
+
+#include "Components/PrimitiveComponent.h"
+#include "GameplayDebuggerRenderingComponent.generated.h"
+
+UCLASS(NotBlueprintable, NotBlueprintType, noteditinlinenew, hidedropdown, Transient)
+class UGameplayDebuggerRenderingComponent : public UPrimitiveComponent
+{
+	GENERATED_UCLASS_BODY()
+
+	virtual FPrimitiveSceneProxy* CreateSceneProxy() override;
+	virtual FBoxSphereBounds CalcBounds(const FTransform &LocalToWorld) const override;
+
+	virtual void CreateRenderState_Concurrent() override;
+	virtual void DestroyRenderState_Concurrent() override;
+};
diff --git a/Engine/Source/Developer/HierarchicalLODUtilities/Private/HierarchicalLODProxyProcessor.cpp b/Engine/Source/Developer/HierarchicalLODUtilities/Private/HierarchicalLODProxyProcessor.cpp
--- a/Engine/Source/Developer/HierarchicalLODUtilities/Private/HierarchicalLODProxyProcessor.cpp
+++ b/Engine/Source/Developer/HierarchicalLODUtilities/Private/HierarchicalLODProxyProcessor.cpp
@@ -32,9 +32,6 @@ bool FHierarchicalLODProxyProcessor::Tick(float DeltaTime)
 		UStaticMesh* MainMesh = nullptr;		
 		for (UObject* AssetObject : Data->AssetObjects)
 		{
-			// We make it private, so it can't be used by outside of map and is invisible in the content browser, and then remove standalone
-			AssetObject->ClearFlags(RF_Public | RF_Standalone);
-			
 			// Check if this is the generated proxy (static-)mesh
 			UStaticMesh* StaticMesh = Cast<UStaticMesh>(AssetObject);
 			if (StaticMesh)
diff --git a/Engine/Source/Developer/HierarchicalLODUtilities/Private/HierarchicalLODUtilities.cpp b/Engine/Source/Developer/HierarchicalLODUtilities/Private/HierarchicalLODUtilities.cpp
--- a/Engine/Source/Developer/HierarchicalLODUtilities/Private/HierarchicalLODUtilities.cpp
+++ b/Engine/Source/Developer/HierarchicalLODUtilities/Private/HierarchicalLODUtilities.cpp
@@ -103,6 +103,29 @@ float FHierarchicalLODUtilities::CalculateDrawDistanceFromScreenSize(const float
 	return Distance;
 }
 
+UPackage* FHierarchicalLODUtilities::CreateOrRetrieveLevelHLODPackage(ULevel* InLevel)
+{
+	checkf(InLevel != nullptr, TEXT("Invalid Level supplied"));
+
+	UPackage* HLODPackage = nullptr;
+	UPackage* LevelOuterMost = InLevel->GetOutermost();
+
+	const FString PathName = FPackageName::GetLongPackagePath(LevelOuterMost->GetPathName());
+	const FString BaseName = FPackageName::GetShortName(LevelOuterMost->GetPathName());
+	const FString HLODLevelPackageName = FString::Printf(TEXT("%s/HLOD/%s_HLOD"), *PathName, *BaseName);
+
+	HLODPackage = CreatePackage(NULL, *HLODLevelPackageName);
+	HLODPackage->FullyLoad();
+	HLODPackage->Modify();
+	
+	// Target level filename
+	const FString HLODLevelFileName = FPackageName::LongPackageNameToFilename(HLODLevelPackageName);
+	// This is a hack to avoid save file dialog when we will be saving HLOD map package
+	HLODPackage->FileName = FName(*HLODLevelFileName);
+
+	return HLODPackage;	
+}
+
 bool FHierarchicalLODUtilities::BuildStaticMeshForLODActor(ALODActor* LODActor, UPackage* AssetsOuter, const FHierarchicalSimplification& LODSetup, const uint32 LODIndex)
 {
 	if (AssetsOuter && LODActor)
@@ -200,13 +223,7 @@ bool FHierarchicalLODUtilities::BuildStaticMeshForLODActor(ALODActor* LODActor,
 				}
 
 				MeshUtilities.MergeStaticMeshComponents(AllComponents, FirstActor->GetWorld(), MergeSettings, AssetsOuter, PackageName, LODIndex, OutAssets, OutProxyLocation, LODSetup.TransitionScreenSize, true);
-
-				// we make it private, so it can't be used by outside of map since it's useless, and then remove standalone
-				for (auto& AssetIter : OutAssets)
-				{
-					AssetIter->ClearFlags(RF_Public | RF_Standalone);
-				}
-
+				
 				// set staticmesh
 				for (auto& Asset : OutAssets)
 				{
diff --git a/Engine/Source/Developer/HierarchicalLODUtilities/Public/HierarchicalLODUtilities.h b/Engine/Source/Developer/HierarchicalLODUtilities/Public/HierarchicalLODUtilities.h
--- a/Engine/Source/Developer/HierarchicalLODUtilities/Public/HierarchicalLODUtilities.h
+++ b/Engine/Source/Developer/HierarchicalLODUtilities/Public/HierarchicalLODUtilities.h
@@ -45,6 +45,9 @@ public:
 
 	static float CalculateDrawDistanceFromScreenSize(const float SphereRadius, const float ScreenSize, const FMatrix& ProjectionMatrix);
 
+	/** Creates or retrieves the HLOD package that is created for the given level */
+	static UPackage* CreateOrRetrieveLevelHLODPackage(ULevel* InLevel);
+
 	/**
 	* Builds a static mesh object for the given LODACtor
 	*
diff --git a/Engine/Source/Developer/LogVisualizer/Private/LogVisualizerSettings.cpp b/Engine/Source/Developer/LogVisualizer/Private/LogVisualizerSettings.cpp
--- a/Engine/Source/Developer/LogVisualizer/Private/LogVisualizerSettings.cpp
+++ b/Engine/Source/Developer/LogVisualizer/Private/LogVisualizerSettings.cpp
@@ -238,8 +238,15 @@ bool FVisualLoggerFilters::MatchCategoryFilters(FString String, ELogVerbosity::T
 		bFoundFilter = Filter.CategoryName == String;
 		if (bFoundFilter)
 		{
-			const bool bFineWithSearchString = Settings->bSearchInsideLogs == true || (SearchBoxFilter.Len() == 0 || Filter.CategoryName.Find(SearchBoxFilter) != INDEX_NONE);
-			return bFineWithSearchString && Verbosity <= Filter.LogVerbosity && Filter.Enabled;
+			if (Filter.Enabled)
+			{
+				const bool bFineWithSearchString = Settings->bSearchInsideLogs == true || (SearchBoxFilter.Len() == 0 || Filter.CategoryName.Find(SearchBoxFilter) != INDEX_NONE);
+				return bFineWithSearchString && Verbosity <= Filter.LogVerbosity;
+			}
+			else
+			{
+				return false;
+			}
 		}
 	}
 
diff --git a/Engine/Source/Developer/LogVisualizer/Private/SVisualLogger.cpp b/Engine/Source/Developer/LogVisualizer/Private/SVisualLogger.cpp
--- a/Engine/Source/Developer/LogVisualizer/Private/SVisualLogger.cpp
+++ b/Engine/Source/Developer/LogVisualizer/Private/SVisualLogger.cpp
@@ -421,7 +421,7 @@ void SVisualLogger::HandleStopRecordingCommandExecute()
 
 	if (FParse::Param(FCommandLine::Get(), TEXT("LogNavOctree")) == true && ULogVisualizerSettings::StaticClass()->GetDefaultObject<ULogVisualizerSettings>()->bLogNavOctreeOnStop)
 	{
-		FVisualLogger::NavigationDataDump(World, LogNavigation, ELogVerbosity::Log, INDEX_NONE, FBox());
+		FVisualLogger::NavigationDataDump(World, LogNavigation, ELogVerbosity::Log, FBox());
 	}
 
 	FVisualLogger::Get().SetIsRecording(false);
@@ -942,9 +942,13 @@ void SVisualLogger::UpdateVisibilityForEntry(const FVisualLoggerDBRow& DBRow, in
 	TArray<FVisualLoggerCategoryVerbosityPair> OutCategories;
 	FVisualLoggerHelpers::GetCategories(CurrentEntry.Entry, OutCategories);
 	bool bHasValidCategories = false;
-	for (FVisualLoggerCategoryVerbosityPair& Categoryair : OutCategories)
+	for (FVisualLoggerCategoryVerbosityPair& CategoryPair : OutCategories)
 	{
-		bHasValidCategories = bHasValidCategories || FVisualLoggerFilters::Get().MatchCategoryFilters(Categoryair.CategoryName.ToString(), Categoryair.Verbosity);
+		if (FVisualLoggerFilters::Get().MatchCategoryFilters(CategoryPair.CategoryName.ToString(), CategoryPair.Verbosity))
+		{
+			bHasValidCategories = true;
+			break;
+		}
 	}
 
 	if (Settings->bSearchInsideLogs && bHasValidCategories && SearchString.Len() > 0)
diff --git a/Engine/Source/Editor/AnimGraph/Classes/AnimGraphNode_BoneDrivenController.h b/Engine/Source/Editor/AnimGraph/Classes/AnimGraphNode_BoneDrivenController.h
--- a/Engine/Source/Editor/AnimGraph/Classes/AnimGraphNode_BoneDrivenController.h
+++ b/Engine/Source/Editor/AnimGraph/Classes/AnimGraphNode_BoneDrivenController.h
@@ -51,7 +51,11 @@ protected:
 	EVisibility AreNonCurveMappingValuesVisible() const;
 	EVisibility AreRemappingValuesVisible() const;
 
-	static void AddTripletPropertyRow(const FText& Name, const FText& Tooltip, IDetailCategoryBuilder& Category, TSharedRef<IPropertyHandle> PropertyHandle, const FName XPropertyName, const FName YPropertyName, const FName ZPropertyName);
+	// Should destination bone or morph target properties be visible
+	EVisibility AreTargetBonePropertiesVisible() const;
+	EVisibility AreTargetCurvePropertiesVisible() const;
+
+	static void AddTripletPropertyRow(const FText& Name, const FText& Tooltip, IDetailCategoryBuilder& Category, TSharedRef<IPropertyHandle> PropertyHandle, const FName XPropertyName, const FName YPropertyName, const FName ZPropertyName, TAttribute<EVisibility> VisibilityAttribute);
 	static void AddRangePropertyRow(const FText& Name, const FText& Tooltip, IDetailCategoryBuilder& Category, TSharedRef<IPropertyHandle> PropertyHandle, const FName MinPropertyName, const FName MaxPropertyName, TAttribute<EVisibility> VisibilityAttribute);
 	static FText ComponentTypeToText(EComponentType::Type Component);
 };
diff --git a/Engine/Source/Editor/AnimGraph/Private/AnimGraphNode_BoneDrivenController.cpp b/Engine/Source/Editor/AnimGraph/Private/AnimGraphNode_BoneDrivenController.cpp
--- a/Engine/Source/Editor/AnimGraph/Private/AnimGraphNode_BoneDrivenController.cpp
+++ b/Engine/Source/Editor/AnimGraph/Private/AnimGraphNode_BoneDrivenController.cpp
@@ -46,7 +46,7 @@ void UAnimGraphNode_BoneDrivenController::Serialize(FArchive& Ar)
 
 FText UAnimGraphNode_BoneDrivenController::GetTooltipText() const
 {
-	return LOCTEXT("UAnimGraphNode_BoneDrivenController_ToolTip", "Drives the transform of a bone using the transform of another");
+	return LOCTEXT("UAnimGraphNode_BoneDrivenController_ToolTip", "Drives the transform of a bone or morph target using the transform of another bone");
 }
 
 FText UAnimGraphNode_BoneDrivenController::GetNodeTitle(ENodeTitleType::Type TitleType) const
@@ -101,6 +101,9 @@ FText UAnimGraphNode_BoneDrivenController::GetNodeTitle(ENodeTitleType::Type Tit
 
 		FFormatNamedArguments Args;
 		Args.Add(TEXT("ControllerDesc"), GetControllerDescription());
+
+		Args.Add(TEXT("ParameterName"), FText::FromName(Node.ParameterName));
+
 		Args.Add(TEXT("TargetBone"), FText::FromName(Node.TargetBone.BoneName));
 
 		// Determine the target component
@@ -129,7 +132,18 @@ FText UAnimGraphNode_BoneDrivenController::GetNodeTitle(ENodeTitleType::Type Tit
 		}
 
 		Args.Add(TEXT("SourceExpression"), FinalSourceExpression);
-		return FText::Format(LOCTEXT("AnimGraphNode_BoneDrivenController_Title", "{TargetBone}.{TargetComponents} = {SourceExpression}{Delim}{ControllerDesc}"), Args);
+
+		FText FormattedText;
+		if (Node.DestinationMode == EDrivenDestinationMode::Bone)
+		{
+			FormattedText = FText::Format(LOCTEXT("AnimGraphNode_BoneDrivenController_Title_Bone", "{TargetBone}.{TargetComponents} = {SourceExpression}{Delim}{ControllerDesc}"), Args);
+		}
+		else if ( Node.DestinationMode == EDrivenDestinationMode::MorphTarget || Node.DestinationMode == EDrivenDestinationMode::MaterialParameter)
+		{
+			FormattedText = FText::Format(LOCTEXT("AnimGraphNode_BoneDrivenController_Title_Curve", "{ParameterName} = {SourceExpression}{Delim}{ControllerDesc}"), Args);
+		}
+
+		return FormattedText;
 	}	
 }
 
@@ -181,24 +195,28 @@ void UAnimGraphNode_BoneDrivenController::ValidateAnimNodeDuringCompilation(USke
 		MessageLog.Warning(*LOCTEXT("NoSourceComponent", "@@ - You must pick a source component on the Driver joint").ToString(), this);
 	}
 	
-	if (ForSkeleton->GetReferenceSkeleton().FindBoneIndex(Node.TargetBone.BoneName) == INDEX_NONE)
+	if (Node.DestinationMode == EDrivenDestinationMode::Bone)
 	{
-		MessageLog.Warning(*LOCTEXT("NoTargetBone", "@@ - You must pick a target bone as the Driven joint").ToString(), this);
-	}
-
-	const bool bAffectsTranslation = Node.bAffectTargetTranslationX || Node.bAffectTargetTranslationY || Node.bAffectTargetTranslationZ;
-	const bool bAffectsRotation = Node.bAffectTargetRotationX || Node.bAffectTargetRotationY || Node.bAffectTargetRotationZ;
-	const bool bAffectsScale = Node.bAffectTargetScaleX || Node.bAffectTargetScaleY || Node.bAffectTargetScaleZ;
+		if (ForSkeleton->GetReferenceSkeleton().FindBoneIndex(Node.TargetBone.BoneName) == INDEX_NONE)
+		{
+			MessageLog.Warning(*LOCTEXT("NoTargetBone", "@@ - You must pick a target bone as the Driven joint").ToString(), this);
+		}
+		
+		const bool bAffectsTranslation = Node.bAffectTargetTranslationX || Node.bAffectTargetTranslationY || Node.bAffectTargetTranslationZ;
+		const bool bAffectsRotation = Node.bAffectTargetRotationX || Node.bAffectTargetRotationY || Node.bAffectTargetRotationZ;
+		const bool bAffectsScale = Node.bAffectTargetScaleX || Node.bAffectTargetScaleY || Node.bAffectTargetScaleZ;
 
-	if (!bAffectsTranslation && !bAffectsRotation && !bAffectsScale)
-	{
-		MessageLog.Warning(*LOCTEXT("NoTargetComponent", "@@ - You must pick one or more target components on the Driven joint").ToString(), this);
+		if (!bAffectsTranslation && !bAffectsRotation && !bAffectsScale)
+		{
+			MessageLog.Warning(*LOCTEXT("NoTargetComponent", "@@ - You must pick one or more target components on the Driven joint").ToString(), this);
+		}
 	}
 
 	Super::ValidateAnimNodeDuringCompilation(ForSkeleton, MessageLog);
 }
 
-void UAnimGraphNode_BoneDrivenController::AddTripletPropertyRow(const FText& Name, const FText& Tooltip, IDetailCategoryBuilder& Category, TSharedRef<IPropertyHandle> PropertyHandle, const FName XPropertyName, const FName YPropertyName, const FName ZPropertyName)
+
+void UAnimGraphNode_BoneDrivenController::AddTripletPropertyRow(const FText& Name, const FText& Tooltip, IDetailCategoryBuilder& Category, TSharedRef<IPropertyHandle> PropertyHandle, const FName XPropertyName, const FName YPropertyName, const FName ZPropertyName, TAttribute<EVisibility> VisibilityAttribute)
 {
 	const float XYZPadding = 5.0f;
 
@@ -212,6 +230,7 @@ void UAnimGraphNode_BoneDrivenController::AddTripletPropertyRow(const FText& Nam
 	Category.GetParentLayout().HideProperty(ZProperty);
 
 	Category.AddCustomRow(Name)
+	.Visibility(VisibilityAttribute)
 	.NameContent()
 	[
 		SNew(STextBlock)
@@ -361,10 +380,19 @@ void UAnimGraphNode_BoneDrivenController::CustomizeDetails(IDetailLayoutBuilder&
 
 	MappingCategory.AddProperty(NodeHandle->GetChildHandle(GET_MEMBER_NAME_CHECKED(FAnimNode_BoneDrivenController, Multiplier))).Visibility(NotUsingCurveVisibility);
 
+	// Destination visiblity
+	TAttribute<EVisibility> BoneTargetVisibility = TAttribute<EVisibility>::Create(TAttribute<EVisibility>::FGetter::CreateUObject(this, &UAnimGraphNode_BoneDrivenController::AreTargetBonePropertiesVisible));
+	TAttribute<EVisibility> CurveTargetVisibility = TAttribute<EVisibility>::Create(TAttribute<EVisibility>::FGetter::CreateUObject(this, &UAnimGraphNode_BoneDrivenController::AreTargetCurvePropertiesVisible));
+
 	// Destination (Driven) category
 	IDetailCategoryBuilder& TargetCategory = DetailBuilder.EditCategory(TEXT("Destination (Driven)"));
-	TargetCategory.AddProperty(NodeHandle->GetChildHandle(GET_MEMBER_NAME_CHECKED(FAnimNode_BoneDrivenController, TargetBone)));
 
+	TargetCategory.AddProperty(NodeHandle->GetChildHandle(GET_MEMBER_NAME_CHECKED(FAnimNode_BoneDrivenController, DestinationMode)));
+
+	TargetCategory.AddProperty(NodeHandle->GetChildHandle(GET_MEMBER_NAME_CHECKED(FAnimNode_BoneDrivenController, ParameterName))).Visibility(CurveTargetVisibility);
+
+
+	TargetCategory.AddProperty(NodeHandle->GetChildHandle(GET_MEMBER_NAME_CHECKED(FAnimNode_BoneDrivenController, TargetBone))).Visibility(BoneTargetVisibility);
 	// Add a note about the space (it is not configurable, and this helps set expectations)
 	const FText TargetBoneSpaceName = LOCTEXT("TargetComponentSpace", "Target Component Space");
 	TargetCategory.AddCustomRow(TargetBoneSpaceName)
@@ -372,14 +400,15 @@ void UAnimGraphNode_BoneDrivenController::CustomizeDetails(IDetailLayoutBuilder&
 		[
 			SNew(STextBlock)
 			.Text(TargetBoneSpaceName)
-			.Font(IDetailLayoutBuilder::GetDetailFont())
+			.Font(IDetailLayoutBuilder::GetDetailFont())			
 		]
 		.ValueContent()
 		[
 			SNew(STextBlock)
 			.Text(LOCTEXT("TargetComponentSpaceIsAlwaysParentBoneSpace", "Parent Bone Space"))
 			.Font(IDetailLayoutBuilder::GetDetailFont())
-		];
+		]
+		.Visibility(BoneTargetVisibility);
 
 	AddTripletPropertyRow(
 		/*Name=*/ LOCTEXT("DrivenTranslationLabel", "Translation"),
@@ -388,7 +417,8 @@ void UAnimGraphNode_BoneDrivenController::CustomizeDetails(IDetailLayoutBuilder&
 		/*PropertyHandle=*/ NodeHandle,
 		/*X=*/ GET_MEMBER_NAME_CHECKED(FAnimNode_BoneDrivenController, bAffectTargetTranslationX),
 		/*Y=*/ GET_MEMBER_NAME_CHECKED(FAnimNode_BoneDrivenController, bAffectTargetTranslationY),
-		/*Z=*/ GET_MEMBER_NAME_CHECKED(FAnimNode_BoneDrivenController, bAffectTargetTranslationZ));
+		/*Z=*/ GET_MEMBER_NAME_CHECKED(FAnimNode_BoneDrivenController, bAffectTargetTranslationZ),
+		BoneTargetVisibility);
 
 	AddTripletPropertyRow(
 		/*Name=*/ LOCTEXT("DrivenRotationLabel", "Rotation"),
@@ -397,7 +427,8 @@ void UAnimGraphNode_BoneDrivenController::CustomizeDetails(IDetailLayoutBuilder&
 		/*PropertyHandle=*/ NodeHandle,
 		/*X=*/ GET_MEMBER_NAME_CHECKED(FAnimNode_BoneDrivenController, bAffectTargetRotationX),
 		/*Y=*/ GET_MEMBER_NAME_CHECKED(FAnimNode_BoneDrivenController, bAffectTargetRotationY),
-		/*Z=*/ GET_MEMBER_NAME_CHECKED(FAnimNode_BoneDrivenController, bAffectTargetRotationZ));
+		/*Z=*/ GET_MEMBER_NAME_CHECKED(FAnimNode_BoneDrivenController, bAffectTargetRotationZ),
+		BoneTargetVisibility);
 
 	AddTripletPropertyRow(
 		/*Name=*/ LOCTEXT("DrivenScaleLabel", "Scale"),
@@ -406,7 +437,10 @@ void UAnimGraphNode_BoneDrivenController::CustomizeDetails(IDetailLayoutBuilder&
 		/*PropertyHandle=*/ NodeHandle,
 		/*X=*/ GET_MEMBER_NAME_CHECKED(FAnimNode_BoneDrivenController, bAffectTargetScaleX),
 		/*Y=*/ GET_MEMBER_NAME_CHECKED(FAnimNode_BoneDrivenController, bAffectTargetScaleY),
-		/*Z=*/ GET_MEMBER_NAME_CHECKED(FAnimNode_BoneDrivenController, bAffectTargetScaleZ));
+		/*Z=*/ GET_MEMBER_NAME_CHECKED(FAnimNode_BoneDrivenController, bAffectTargetScaleZ),
+		BoneTargetVisibility);
+
+	TargetCategory.AddProperty(NodeHandle->GetChildHandle(GET_MEMBER_NAME_CHECKED(FAnimNode_BoneDrivenController, ModificationMode))).Visibility(BoneTargetVisibility);
 }
 
 FText UAnimGraphNode_BoneDrivenController::ComponentTypeToText(EComponentType::Type Component)
@@ -448,4 +482,14 @@ EVisibility UAnimGraphNode_BoneDrivenController::AreRemappingValuesVisible() con
 	return ((Node.DrivingCurve == nullptr) && Node.bUseRange) ? EVisibility::Visible : EVisibility::Collapsed;
 }
 
+EVisibility UAnimGraphNode_BoneDrivenController::AreTargetBonePropertiesVisible() const
+{
+	return (Node.DestinationMode == EDrivenDestinationMode::Bone) ? EVisibility::Visible : EVisibility::Collapsed;
+}
+
+EVisibility UAnimGraphNode_BoneDrivenController::AreTargetCurvePropertiesVisible() const
+{
+	return (Node.DestinationMode == EDrivenDestinationMode::MorphTarget || Node.DestinationMode == EDrivenDestinationMode::MaterialParameter) ? EVisibility::Visible : EVisibility::Collapsed;
+}
+
 #undef LOCTEXT_NAMESPACE
diff --git a/Engine/Source/Editor/DetailCustomizations/Private/BodyInstanceCustomization.cpp b/Engine/Source/Editor/DetailCustomizations/Private/BodyInstanceCustomization.cpp
--- a/Engine/Source/Editor/DetailCustomizations/Private/BodyInstanceCustomization.cpp
+++ b/Engine/Source/Editor/DetailCustomizations/Private/BodyInstanceCustomization.cpp
@@ -21,13 +21,38 @@ FBodyInstanceCustomization::FBodyInstanceCustomization()
 	RefreshCollisionProfiles();
 }
 
+UStaticMeshComponent* GetDefaultCollisionProvider(const FBodyInstance* BI)
+{
+	UStaticMeshComponent* SMC = BI ? Cast<UStaticMeshComponent>(BI->OwnerComponent.Get()) : nullptr;
+	return SMC && SMC->SupportsDefaultCollision() ? SMC : nullptr;
+}
+
+bool FBodyInstanceCustomization::CanUseDefaultCollision() const
+{
+	bool bResult = BodyInstances.Num() > 0;
+	for(const FBodyInstance* BI : BodyInstances)
+	{
+		bResult &= GetDefaultCollisionProvider(BI) != nullptr;
+	}
+
+	return bResult;
+}
+
 void FBodyInstanceCustomization::RefreshCollisionProfiles()
 {
 	int32 NumProfiles = CollisionProfile->GetNumOfProfiles();
+
+	bool bCanUseDefaultCollision = CanUseDefaultCollision();
+
 	// first create profile combo list
-	CollisionProfileComboList.Empty(NumProfiles + 1);
+	CollisionProfileComboList.Empty(NumProfiles + (bCanUseDefaultCollision ? 2 : 1));	//If we can use default collision we'll add a "Default" option
 
 	// first one is default one
+	if(bCanUseDefaultCollision)
+	{
+		CollisionProfileComboList.Add(MakeShareable(new FString(TEXT("Default"))));
+	}
+	
 	CollisionProfileComboList.Add(MakeShareable(new FString(TEXT("Custom..."))));
 
 	// go through profile and see if it has mine
@@ -35,6 +60,11 @@ void FBodyInstanceCustomization::RefreshCollisionProfiles()
 	{
 		CollisionProfileComboList.Add(MakeShareable(new FString(CollisionProfile->GetProfileByIndex(ProfileId)->Name.ToString())));
 	}
+
+	if(CollsionProfileComboBox.IsValid())
+	{
+		CollsionProfileComboBox->RefreshOptions();
+	}
 }
 
 void FBodyInstanceCustomization::AddCollisionCategory(TSharedRef<class IPropertyHandle> StructPropertyHandle, class IDetailChildrenBuilder& StructBuilder, IPropertyTypeCustomizationUtils& StructCustomizationUtils)
@@ -55,7 +85,7 @@ void FBodyInstanceCustomization::AddCollisionCategory(TSharedRef<class IProperty
 	bool bDisplayAdvancedCollisionSettings = true;
 
 	// if I have valid profile name
-	if (CollisionProfileNameHandle->GetValue(ProfileName) ==  FPropertyAccess::Result::Success && FBodyInstance::IsValidCollisionProfileName(ProfileName) )
+	if (!AreAllCollisionUsingDefault() && CollisionProfileNameHandle->GetValue(ProfileName) ==  FPropertyAccess::Result::Success && FBodyInstance::IsValidCollisionProfileName(ProfileName) )
 	{
 		DisplayName = GetProfileString(ProfileName);
 		bDisplayAdvancedCollisionSettings = false;
@@ -80,7 +110,6 @@ void FBodyInstanceCustomization::AddCollisionCategory(TSharedRef<class IProperty
 		.Padding(0.f, 0.f, 10.f, 0.f)
 		[
 			SNew(SHorizontalBox)
-			.Visibility(this, &FBodyInstanceCustomization::IsCollisionPresetVisible)
 			.IsEnabled(this, &FBodyInstanceCustomization::IsCollisionEnabled)
 			+ SHorizontalBox::Slot()
 			.VAlign(VAlign_Center)
@@ -102,36 +131,24 @@ void FBodyInstanceCustomization::AddCollisionCategory(TSharedRef<class IProperty
 			]
 
 			+ SHorizontalBox::Slot()
-				.VAlign(VAlign_Center)
-				.Padding(2.0f)
-				.AutoWidth()
+			.VAlign(VAlign_Center)
+			.Padding(2.0f)
+			.AutoWidth()
+			[
+				SNew(SButton)
+				.OnClicked(this, &FBodyInstanceCustomization::SetToDefaultProfile)
+				.ContentPadding(0.f)
+				.ToolTipText(LOCTEXT("ResetToDefaultToolTip", "Reset to Default"))
+				.ButtonStyle(FEditorStyle::Get(), "NoBorder")
+				.IsEnabled(this, &FBodyInstanceCustomization::IsCollisionEnabled)
+				.Visibility(this, &FBodyInstanceCustomization::ShouldShowResetToDefaultProfile)
+				.Content()
 				[
-					SNew(SButton)
-					.OnClicked(this, &FBodyInstanceCustomization::SetToDefaultProfile)
-					.ContentPadding(0.f)
-					.ToolTipText(LOCTEXT("ResetToDefaultToolTip", "Reset to Default"))
-					.ButtonStyle(FEditorStyle::Get(), "NoBorder")
-					.IsEnabled(this, &FBodyInstanceCustomization::IsCollisionEnabled)
-					.Content()
-					[
-						SNew(SImage)
-						.Image(FEditorStyle::GetBrush("PropertyWindow.DiffersFromDefault"))
-					]
+					SNew(SImage)
+					.Image(FEditorStyle::GetBrush("PropertyWindow.DiffersFromDefault"))
 				]
-		]
-
-		+ SVerticalBox::Slot()
-			.Padding(0.f, 0.f, 10.f, 0.f)
-			[
-				SNew(STextBlock)
-				.Text(LOCTEXT("DefaultCollision", "Default Collision"))
-				.Visibility(this, &FBodyInstanceCustomization::IsDefaultCollisionVisible)
-				.Font(IDetailLayoutBuilder::GetDetailFont())
-				.ToolTipText(LOCTEXT("DefaultCollisionToolTip", "Default Collision. See the StaticMesh asset to find out collision settings"))
 			]
-
-
-		
+		]
 	];
 
 	CollisionGroup.ToggleExpansion(bDisplayAdvancedCollisionSettings);
@@ -156,6 +173,7 @@ void FBodyInstanceCustomization::CustomizeChildren( TSharedRef<class IPropertyHa
 		BodyInstances[Iter.GetIndex()] = (FBodyInstance*)(*Iter);
 	}
 
+	RefreshCollisionProfiles();
 	
 	// get all parent instances
 	TSharedPtr<IPropertyHandle> ParentPropertyHandle = StructPropertyHandle->GetParentHandle();
@@ -178,6 +196,12 @@ void FBodyInstanceCustomization::CustomizeChildren( TSharedRef<class IPropertyHa
 		
 		UseDefaultCollisionHandle = ParentPropertyHandle->GetChildHandle(GET_MEMBER_NAME_CHECKED(UStaticMeshComponent, bUseDefaultCollision));
 		StaticMeshHandle = ParentPropertyHandle->GetChildHandle(GET_MEMBER_NAME_CHECKED(UStaticMeshComponent, StaticMesh));
+
+		if(StaticMeshHandle.IsValid())
+		{
+			FSimpleDelegate OnStaticMeshChangedDelegate = FSimpleDelegate::CreateSP(this, &FBodyInstanceCustomization::RefreshCollisionProfiles);
+			StaticMeshHandle->SetOnPropertyValueChanged(OnStaticMeshChangedDelegate);
+		}
 	}
 	
 
@@ -235,6 +259,22 @@ int32 FBodyInstanceCustomization::InitializeObjectTypeComboList()
 	return Selected;
 }
 
+int32 FBodyInstanceCustomization::GetNumberOfSpecialProfiles() const
+{
+	return CanUseDefaultCollision() ? 2 : 1;
+}
+
+int32 FBodyInstanceCustomization::GetCustomIndex() const
+{
+	return CanUseDefaultCollision() ? 1 : 0;
+}
+
+int32 FBodyInstanceCustomization::GetDefaultIndex() const
+{
+	ensure(CanUseDefaultCollision());
+	return 0;
+}
+
 TSharedPtr<FString> FBodyInstanceCustomization::GetProfileString(FName ProfileName) const
 {
 	FString ProfileNameString = ProfileName.ToString();
@@ -242,7 +282,7 @@ TSharedPtr<FString> FBodyInstanceCustomization::GetProfileString(FName ProfileNa
 	// go through profile and see if it has mine
 	int32 NumProfiles = CollisionProfile->GetNumOfProfiles();
 	// refresh collision count
-	if( NumProfiles+1==CollisionProfileComboList.Num() )
+	if( NumProfiles + GetNumberOfSpecialProfiles() == CollisionProfileComboList.Num() )
 	{
 		for(int32 ProfileId = 0; ProfileId < NumProfiles; ++ProfileId)
 		{
@@ -253,7 +293,7 @@ TSharedPtr<FString> FBodyInstanceCustomization::GetProfileString(FName ProfileNa
 		}
 	}
 
-	return CollisionProfileComboList[0];
+	return CollisionProfileComboList[GetCustomIndex()];
 }
 
 // filter through find valid index of enum values matching each item
@@ -718,6 +758,16 @@ TSharedRef<SWidget> FBodyInstanceCustomization::MakeCollisionProfileComboWidget(
 
 void FBodyInstanceCustomization::OnCollisionProfileComboOpening()
 {
+	if(AreAllCollisionUsingDefault())
+	{
+		TSharedPtr<FString> ComboStringPtr = CollisionProfileComboList[GetDefaultIndex()];
+		if (ComboStringPtr.IsValid())
+		{
+			CollsionProfileComboBox->SetSelectedItem(ComboStringPtr);
+			return;
+		}
+	}
+
 	FName ProfileName;
 	if (CollisionProfileNameHandle->GetValue(ProfileName) != FPropertyAccess::Result::MultipleValues)
 	{
@@ -729,6 +779,20 @@ void FBodyInstanceCustomization::OnCollisionProfileComboOpening()
 	}
 }
 
+void MarkAllBodiesDefaultCollision(TArray<FBodyInstance*>& BodyInstances, bool bUseDefaultCollision)
+{
+	for (const FBodyInstance* BI : BodyInstances)
+	{
+		if (UStaticMeshComponent* SMC = GetDefaultCollisionProvider(BI))
+		{
+			if(SMC->SupportsDefaultCollision())
+			{
+				SMC->bUseDefaultCollision = bUseDefaultCollision;
+			}
+		}
+	}
+}
+
 void FBodyInstanceCustomization::OnCollisionProfileChanged( TSharedPtr<FString> NewSelection, ESelectInfo::Type SelectInfo, IDetailGroup* CollisionGroup )
 {
 	// if it's set from code, we did that on purpose
@@ -746,18 +810,31 @@ void FBodyInstanceCustomization::OnCollisionProfileChanged( TSharedPtr<FString>
 				// set profile set up
 				ensure ( CollisionProfileNameHandle->SetValue(NewValue) ==  FPropertyAccess::Result::Success );
 				UpdateCollisionProfile();
+				MarkAllBodiesDefaultCollision(BodyInstances, false);
 				return;
 			}
 		}
 
-		if( NewSelection == CollisionProfileComboList[0] )
+		if(CanUseDefaultCollision())
+		{
+			if(NewSelection == CollisionProfileComboList[GetDefaultIndex()])
+			{
+				MarkAllBodiesDefaultCollision(BodyInstances, true);
+				return;
+			}
+		}
+
+		if( NewSelection == CollisionProfileComboList[GetCustomIndex()])
 		{
 			// Force expansion when the user chooses the selected item
 			CollisionGroup->ToggleExpansion( true );
 		}
+
 		// if none of them found, clear it
 		FName Name=UCollisionProfile::CustomCollisionProfileName;
 		ensure ( CollisionProfileNameHandle->SetValue(Name) ==  FPropertyAccess::Result::Success );
+
+		MarkAllBodiesDefaultCollision(BodyInstances, false);
 	}
 }
 
@@ -766,9 +843,10 @@ void FBodyInstanceCustomization::UpdateCollisionProfile()
 	FName ProfileName;
 
 	// if I have valid profile name
-	if (CollisionProfileNameHandle->GetValue(ProfileName) ==  FPropertyAccess::Result::Success && FBodyInstance::IsValidCollisionProfileName(ProfileName) )
+	if (!AreAllCollisionUsingDefault() && CollisionProfileNameHandle->GetValue(ProfileName) ==  FPropertyAccess::Result::Success && FBodyInstance::IsValidCollisionProfileName(ProfileName) )
 	{
 		int32 NumProfiles = CollisionProfile->GetNumOfProfiles();
+		const int32 NumSpecialProfiles = GetNumberOfSpecialProfiles();
 		for (int32 ProfileId = 0; ProfileId < NumProfiles; ++ProfileId)
 		{
 			// find the profile
@@ -782,7 +860,7 @@ void FBodyInstanceCustomization::UpdateCollisionProfile()
 				SetCollisionResponseContainer(CurProfile->ResponseToChannels);
 
 				// now update combo box
-				CollsionProfileComboBox.Get()->SetSelectedItem(CollisionProfileComboList[ProfileId+1]);
+				CollsionProfileComboBox.Get()->SetSelectedItem(CollisionProfileComboList[ProfileId+NumSpecialProfiles]);
 				if (ObjectTypeComboBox.IsValid())
 				{
 					for (auto Iter = ObjectTypeValues.CreateConstIterator(); Iter; ++Iter)
@@ -800,13 +878,14 @@ void FBodyInstanceCustomization::UpdateCollisionProfile()
 		}
 	}
 
-	CollsionProfileComboBox.Get()->SetSelectedItem(CollisionProfileComboList[0]);
+	CollsionProfileComboBox.Get()->SetSelectedItem(CollisionProfileComboList[AreAllCollisionUsingDefault() ? GetDefaultIndex() : GetCustomIndex()]);
 }
 
 FReply FBodyInstanceCustomization::SetToDefaultProfile()
 {
 	// trigger transaction before UpdateCollisionProfile
 	const FScopedTransaction Transaction( LOCTEXT( "ResetCollisionProfile", "Reset Collision Profile" ) );
+	MarkAllBodiesDefaultCollision(BodyInstances, false);
 	CollisionProfileNameHandle.Get()->ResetToDefault();
 	UpdateCollisionProfile();
 	return FReply::Handled();
@@ -853,35 +932,21 @@ EVisibility FBodyInstanceCustomization::ShouldShowResetToDefaultResponse(int32 V
 
 bool FBodyInstanceCustomization::AreAllCollisionUsingDefault() const
 {
-	if(UseDefaultCollisionHandle.IsValid())
+	bool bAllUsingDefault = BodyInstances.Num() > 0;
+	for(const FBodyInstance* BI : BodyInstances)
 	{
-		bool bUseDefaultCollision = false;
-		if (UseDefaultCollisionHandle->GetValue(bUseDefaultCollision) == FPropertyAccess::Result::Success && bUseDefaultCollision)
+		if(UStaticMeshComponent* SMC = GetDefaultCollisionProvider(BI))
 		{
-			return true;
+			bAllUsingDefault &= SMC->bUseDefaultCollision;
 		}
-	}
-	
-	return false;
-}
-
-EVisibility FBodyInstanceCustomization::IsCollisionPresetVisible() const
-{
-	if (StaticMeshHandle.IsValid())
-	{
-		UObject* StaticMesh = nullptr;
-		if (StaticMeshHandle->GetValue(StaticMesh) == FPropertyAccess::Result::Success && StaticMesh)
+		else
 		{
-			return AreAllCollisionUsingDefault() ? EVisibility::Collapsed : EVisibility::Visible;
+			bAllUsingDefault = false;
+			break;
 		}
 	}
 
-	return EVisibility::Visible;
-}
-
-EVisibility FBodyInstanceCustomization::IsDefaultCollisionVisible() const
-{
-	return IsCollisionPresetVisible()  == EVisibility::Visible ? EVisibility::Collapsed : EVisibility::Visible;
+	return bAllUsingDefault;
 }
 
 bool FBodyInstanceCustomization::IsCollisionEnabled() const
@@ -898,7 +963,7 @@ bool FBodyInstanceCustomization::IsCollisionEnabled() const
 bool FBodyInstanceCustomization::ShouldEnableCustomCollisionSetup() const
 {
 	FName ProfileName;
-	if (CollisionProfileNameHandle->GetValue(ProfileName) == FPropertyAccess::Result::Success && FBodyInstance::IsValidCollisionProfileName(ProfileName) == false)
+	if (!AreAllCollisionUsingDefault() && CollisionProfileNameHandle->GetValue(ProfileName) == FPropertyAccess::Result::Success && FBodyInstance::IsValidCollisionProfileName(ProfileName) == false)
 	{
 		return IsCollisionEnabled();
 	}
@@ -913,8 +978,29 @@ EVisibility FBodyInstanceCustomization::ShouldShowCustomCollisionSetup() const
 
 FText FBodyInstanceCustomization::GetCollisionProfileComboBoxContent() const
 {
+	bool bAllUseDefaultCollision = BodyInstances.Num() > 0;
+	bool bSomeUseDefaultCollision = false;
+
+	for(const FBodyInstance* BI : BodyInstances)
+	{
+		if(UStaticMeshComponent* SMC = GetDefaultCollisionProvider(BI))
+		{
+			bAllUseDefaultCollision &= SMC->bUseDefaultCollision;
+			bSomeUseDefaultCollision |= SMC->bUseDefaultCollision;
+		}
+		else
+		{
+			bAllUseDefaultCollision = false;
+		}
+	}
+
+	if(bAllUseDefaultCollision)
+	{
+		return FText::FromString(*CollisionProfileComboList[GetDefaultIndex()]);
+	}
+
 	FName ProfileName;
-	if (CollisionProfileNameHandle->GetValue(ProfileName) == FPropertyAccess::Result::MultipleValues)
+	if (bSomeUseDefaultCollision || CollisionProfileNameHandle->GetValue(ProfileName) == FPropertyAccess::Result::MultipleValues)
 	{
 		return LOCTEXT("MultipleValues", "Multiple Values");
 	}
@@ -924,6 +1010,11 @@ FText FBodyInstanceCustomization::GetCollisionProfileComboBoxContent() const
 
 FText FBodyInstanceCustomization::GetCollisionProfileComboBoxToolTip() const
 {
+	if(AreAllCollisionUsingDefault())
+	{
+		return LOCTEXT("DefaultCollision", "Default collision preset specified in the StaticMesh asset");
+	}
+
 	FName ProfileName;
 	if (CollisionProfileNameHandle->GetValue(ProfileName) == FPropertyAccess::Result::Success)
 	{
diff --git a/Engine/Source/Editor/DetailCustomizations/Private/BodyInstanceCustomization.h b/Engine/Source/Editor/DetailCustomizations/Private/BodyInstanceCustomization.h
--- a/Engine/Source/Editor/DetailCustomizations/Private/BodyInstanceCustomization.h
+++ b/Engine/Source/Editor/DetailCustomizations/Private/BodyInstanceCustomization.h
@@ -61,8 +61,6 @@ private:
 	bool IsCollisionEnabled() const;
 
 	// whether we can edit collision or if we're getting it from a default
-	EVisibility IsCollisionPresetVisible() const;
-	EVisibility IsDefaultCollisionVisible() const;
 	bool AreAllCollisionUsingDefault() const;
 
 	// utility functions between property and struct
@@ -105,6 +103,12 @@ private:
 	TArray<FCollisionChannelInfo>	ValidCollisionChannels;
 
 	void RefreshCollisionProfiles();
+
+	bool CanUseDefaultCollision() const;
+	bool CanShowDefaultCollision() const;
+	int32 GetNumberOfSpecialProfiles() const;
+	int32 GetCustomIndex() const;
+	int32 GetDefaultIndex() const;
 };
 
 class FBodyInstanceCustomizationHelper  : public TSharedFromThis<FBodyInstanceCustomizationHelper>
diff --git a/Engine/Source/Editor/DetailCustomizations/Private/StaticMeshComponentDetails.cpp b/Engine/Source/Editor/DetailCustomizations/Private/StaticMeshComponentDetails.cpp
--- a/Engine/Source/Editor/DetailCustomizations/Private/StaticMeshComponentDetails.cpp
+++ b/Engine/Source/Editor/DetailCustomizations/Private/StaticMeshComponentDetails.cpp
@@ -15,46 +15,6 @@ void FStaticMeshComponentDetails::CustomizeDetails( IDetailLayoutBuilder& Detail
 {
 	// Create a category so this is displayed early in the properties
 	DetailBuilder.EditCategory( "StaticMesh", FText::GetEmpty(), ECategoryPriority::Important);
-
-	DetailBuilder.GetObjectsBeingCustomized(ObjectsCustomized);
-	
-	TSharedPtr<IPropertyHandle> UseDefaultCollisionHandle = DetailBuilder.GetProperty(GET_MEMBER_NAME_CHECKED(UStaticMeshComponent, bUseDefaultCollision));
-
-	DetailBuilder.EditCategory(TEXT("Collision")).AddCustomRow(UseDefaultCollisionHandle->GetPropertyDisplayName())
-	.IsEnabled(TAttribute<bool>(this, &FStaticMeshComponentDetails::IsDefaultCollisionSupported))
-	.NameContent()
-	[
-		UseDefaultCollisionHandle->CreatePropertyNameWidget()
-	]
-	.ValueContent()
-	[
-		UseDefaultCollisionHandle->CreatePropertyValueWidget()
-	];
-
-	UseDefaultCollisionHandle->MarkHiddenByCustomization();
-}
-
-bool FStaticMeshComponentDetails::IsDefaultCollisionSupported() const
-{
-	// Whether to show the default collision property
-	bool bShowCollisionDefault = false;
-	for (TWeakObjectPtr<UObject> Obj : ObjectsCustomized)
-	{
-		if (UStaticMeshComponent* SMC = Cast<UStaticMeshComponent>(Obj.Get()))
-		{
-			if (SMC->SupportsDefaultCollision())
-			{
-				bShowCollisionDefault = true;
-			}
-			else
-			{
-				bShowCollisionDefault = false;
-				break;
-			}
-		}
-	}
-
-	return bShowCollisionDefault;
 }
 
 #undef LOCTEXT_NAMESPACE
diff --git a/Engine/Source/Editor/UnrealEd/Private/Commandlets/PackageUtilities.cpp b/Engine/Source/Editor/UnrealEd/Private/Commandlets/PackageUtilities.cpp
--- a/Engine/Source/Editor/UnrealEd/Private/Commandlets/PackageUtilities.cpp
+++ b/Engine/Source/Editor/UnrealEd/Private/Commandlets/PackageUtilities.cpp
@@ -2475,11 +2475,7 @@ int32 UReplaceActorCommandlet::Main(const FString& Params)
 
 						if (ClassToReplace->IsChildOf(AWorldSettings::StaticClass()))
 						{
-							// Find the index in the array the worldsettings has been spawned at. 
-							const int32 WorldSettingsActorIndex = Level->Actors.Find( NewActor );
-
-							// The worldsettings needs to reside at index 0.
-							Exchange(Level->Actors[0],Level ->Actors[WorldSettingsActorIndex]);
+							Level->SetWorldSettings(CastChecked<AWorldSettings>(NewActor));
 						}
 						check(OldActor->IsValidLowLevel()); // make sure DestroyActor() doesn't immediately trigger GC since that'll break the reference replacement
 						// check for any references to the old Actor and replace them with the new one
diff --git a/Engine/Source/Editor/UnrealEd/Private/EditorLevelUtils.cpp b/Engine/Source/Editor/UnrealEd/Private/EditorLevelUtils.cpp
--- a/Engine/Source/Editor/UnrealEd/Private/EditorLevelUtils.cpp
+++ b/Engine/Source/Editor/UnrealEd/Private/EditorLevelUtils.cpp
@@ -217,27 +217,22 @@ namespace EditorLevelUtils
 			InWorld->MarkPackageDirty();
 
 			NewLevel = StreamingLevel->GetLoadedLevel();
-			if ( NewLevel != NULL )
+			if ( NewLevel != nullptr)
 			{
 				EditorLevelUtils::SetLevelVisibility( NewLevel, true, true );
-			}
 
-			// Levels migrated from other projects may fail to load their world settings
-			// If so we create a new AWorldSettings actor here.
-			if ( NewLevel->Actors[0] == NULL )
-			{
-				UWorld *SubLevelWorld = CastChecked<UWorld>(NewLevel->GetOuter());
-				if (SubLevelWorld != NULL)
+				// Levels migrated from other projects may fail to load their world settings
+				// If so we create a new AWorldSettings actor here.
+				if ( NewLevel->GetWorldSettings(false) == nullptr )
 				{
+					UWorld* SubLevelWorld = CastChecked<UWorld>(NewLevel->GetOuter());
+
 					FActorSpawnParameters SpawnInfo;
 					SpawnInfo.SpawnCollisionHandlingOverride = ESpawnActorCollisionHandlingMethod::AlwaysSpawn;
 					SpawnInfo.Name = GEngine->WorldSettingsClass->GetFName();
 					AWorldSettings* NewWorldSettings = SubLevelWorld->SpawnActor<AWorldSettings>( GEngine->WorldSettingsClass, SpawnInfo );
-					NewLevel->Actors[0] = NewWorldSettings;
-				}
-				else
-				{
-					FMessageDialog::Open( EAppMsgType::Ok, NSLOCTEXT("UnrealEd", "LevelHasNoWorldSettings", "AddLevelToWorld: The level has no World Settings.") );
+
+					NewLevel->SetWorldSettings(NewWorldSettings);
 				}
 			}
 		}
diff --git a/Engine/Source/Editor/UnrealEd/Private/EditorServer.cpp b/Engine/Source/Editor/UnrealEd/Private/EditorServer.cpp
--- a/Engine/Source/Editor/UnrealEd/Private/EditorServer.cpp
+++ b/Engine/Source/Editor/UnrealEd/Private/EditorServer.cpp
@@ -3054,8 +3054,8 @@ ULevel*  UEditorEngine::CreateTransLevelMoveBuffer( UWorld* InWorld )
 	// Spawn worldsettings.
 	FActorSpawnParameters SpawnInfo;
 	SpawnInfo.OverrideLevel = BufferLevel;
-	InWorld->SpawnActor( GEngine->WorldSettingsClass, NULL, NULL, SpawnInfo );
-	check( Cast<AWorldSettings>( BufferLevel->Actors[0] ) );
+	AWorldSettings* WorldSettings = InWorld->SpawnActor<AWorldSettings>( GEngine->WorldSettingsClass, SpawnInfo );
+	BufferLevel->SetWorldSettings(WorldSettings);
 
 	// Spawn builder brush for the buffer level.
 	ABrush* BufferDefaultBrush = InWorld->SpawnActor<ABrush>( SpawnInfo );
@@ -3693,7 +3693,7 @@ bool UEditorEngine::Map_Check( UWorld* InWorld, const TCHAR* Str, FOutputDevice&
 			if (Level != NULL)
 			{
 				// Grab the world info of the streaming level, and loop through it's streaming levels
-				AWorldSettings* SubLevelWorldSettings = CastChecked<AWorldSettings>(Level->Actors[0]);
+				AWorldSettings* SubLevelWorldSettings = Level->GetWorldSettings();
 				UWorld *SubLevelWorld = CastChecked<UWorld>(Level->GetOuter());
 				if (SubLevelWorld != NULL && SubLevelWorldSettings != NULL )
 				{
diff --git a/Engine/Source/Editor/UnrealEd/Private/HierarchicalLOD.cpp b/Engine/Source/Editor/UnrealEd/Private/HierarchicalLOD.cpp
--- a/Engine/Source/Editor/UnrealEd/Private/HierarchicalLOD.cpp
+++ b/Engine/Source/Editor/UnrealEd/Private/HierarchicalLOD.cpp
@@ -464,7 +464,7 @@ void FHierarchicalLODBuilder::BuildMeshesForLODActors()
 
 		if (LevelIter->Actors.Num() > 0)
 		{
-			UPackage* AssetsOuter = LevelIter->GetOutermost();
+			UPackage* AssetsOuter = FHierarchicalLODUtilities::CreateOrRetrieveLevelHLODPackage(LevelIter);
 
 			if (AssetsOuter)
 			{
@@ -549,11 +549,13 @@ void FHierarchicalLODBuilder::DeleteLODActors(ULevel* InLevel, const bool bPrevi
 	CollectGarbage(GARBAGE_COLLECTION_KEEPFLAGS, true);
 }
 
+
+
 void FHierarchicalLODBuilder::BuildMeshForLODActor(ALODActor* LODActor, const uint32 LODLevel)
 {
 	AWorldSettings* WorldSetting = World->GetWorldSettings();
 	BuildLODLevelSettings = WorldSetting->HierarchicalLODSetup;
-	UPackage* AssetsOuter = LODActor->GetLevel()->GetOutermost();
+	UPackage* AssetsOuter = FHierarchicalLODUtilities::CreateOrRetrieveLevelHLODPackage(LODActor->GetLevel());
 
 	const bool bResult = FHierarchicalLODUtilities::BuildStaticMeshForLODActor(LODActor, AssetsOuter, BuildLODLevelSettings[LODLevel], LODLevel);
 
diff --git a/Engine/Source/Editor/UnrealEd/Private/LODCluster.cpp b/Engine/Source/Editor/UnrealEd/Private/LODCluster.cpp
--- a/Engine/Source/Editor/UnrealEd/Private/LODCluster.cpp
+++ b/Engine/Source/Editor/UnrealEd/Private/LODCluster.cpp
@@ -272,7 +272,8 @@ ALODActor* FLODCluster::BuildActor(ULevel* InLevel, const int32 LODIdx, const bo
 		const int32 LODCount = InLevel->GetWorld()->GetWorldSettings()->HierarchicalLODSetup.Num();
 
 		// Where generated assets will be stored
-		UPackage* AssetsOuter = InLevel->GetOutermost(); // this asset is going to save with map, this means, I'll have to delete with it
+		UPackage* AssetsOuter = FHierarchicalLODUtilities::CreateOrRetrieveLevelHLODPackage(InLevel);
+
 		if (AssetsOuter)
 		{
 			TArray<UStaticMeshComponent*> AllComponents;
diff --git a/Engine/Source/Editor/UnrealEd/Private/Lightmass/Lightmass.cpp b/Engine/Source/Editor/UnrealEd/Private/Lightmass/Lightmass.cpp
--- a/Engine/Source/Editor/UnrealEd/Private/Lightmass/Lightmass.cpp
+++ b/Engine/Source/Editor/UnrealEd/Private/Lightmass/Lightmass.cpp
@@ -3379,7 +3379,7 @@ void FLightmassProcessor::ImportMeshAreaLightData()
 					Direction = LMCurrentLightData.Direction;
 					// Spawn a AGeneratedMeshAreaLight to handle the light's influence on dynamic objects
 					FActorSpawnParameters SpawnInfo;
-					SpawnInfo.Owner = CurrentLevel->Actors[0];
+					SpawnInfo.Owner = CurrentLevel->GetWorldSettings();
 					AGeneratedMeshAreaLight* NewGeneratedLight = CurrentLevel->OwningWorld->SpawnActor<AGeneratedMeshAreaLight>(Position, Direction.Rotation());
 					USpotLightComponent* SpotComponent = CastChecked<USpotLightComponent>(NewGeneratedLight->GetLightComponent());
 					// Unregister the component before we change its attributes
diff --git a/Engine/Source/Runtime/AIModule/AIModule.Build.cs b/Engine/Source/Runtime/AIModule/AIModule.Build.cs
--- a/Engine/Source/Runtime/AIModule/AIModule.Build.cs
+++ b/Engine/Source/Runtime/AIModule/AIModule.Build.cs
@@ -63,6 +63,12 @@ namespace UnrealBuildTool.Rules
                 // module, so it's definitions won't propagate to modules that import Engine.
                 Definitions.Add("WITH_RECAST=0");
             }
-		}
-	}
+
+            if (UEBuildConfiguration.bBuildDeveloperTools && Target.Configuration != UnrealTargetConfiguration.Shipping && Target.Configuration != UnrealTargetConfiguration.Test)
+            {
+                PrivateDependencyModuleNames.Add("GameplayDebugger");
+                Definitions.Add("WITH_GAMEPLAY_DEBUGGER=1");
+            }
+        }
+    }
 }
diff --git a/Engine/Source/Runtime/AIModule/Classes/AISystem.h b/Engine/Source/Runtime/AIModule/Classes/AISystem.h
--- a/Engine/Source/Runtime/AIModule/Classes/AISystem.h
+++ b/Engine/Source/Runtime/AIModule/Classes/AISystem.h
@@ -64,6 +64,10 @@ public:
 	UPROPERTY(globalconfig, EditDefaultsOnly, Category = "EQS")
 	bool bAllowControllersAsEQSQuerier;
 
+	/** if set, GameplayDebuggerPlugin will be loaded on module's startup */
+	UPROPERTY(globalconfig, EditDefaultsOnly, Category = "AISystem")
+	bool bEnableDebuggerPlugin;
+
 protected:
 	/** Behavior tree manager used by game */
 	UPROPERTY(Transient)
@@ -102,7 +106,7 @@ public:
 	virtual void CleanupWorld(bool bSessionEnded = true, bool bCleanupResources = true, UWorld* NewWorld = NULL) override;
 	virtual void StartPlay() override;
 	// UAISystemBase end
-	
+
 	/** Behavior tree manager getter */
 	FORCEINLINE UBehaviorTreeManager* GetBehaviorTreeManager() { return BehaviorTreeManager; }
 	/** Behavior tree manager const getter */
@@ -221,6 +225,9 @@ public:
 	*/
 	FBlackboardDataToComponentsIterator CreateBlackboardDataToComponentsIterator(class UBlackboardData& BlackboardAsset);
 
+	virtual void ConditionalLoadDebuggerPlugin();
+
 protected:
 	virtual void OnActorSpawned(AActor* SpawnedActor);
+	void LoadDebuggerPlugin();
 };
diff --git a/Engine/Source/Runtime/AIModule/Classes/EnvironmentQuery/EQSRenderingComponent.h b/Engine/Source/Runtime/AIModule/Classes/EnvironmentQuery/EQSRenderingComponent.h
--- a/Engine/Source/Runtime/AIModule/Classes/EnvironmentQuery/EQSRenderingComponent.h
+++ b/Engine/Source/Runtime/AIModule/Classes/EnvironmentQuery/EQSRenderingComponent.h
@@ -5,6 +5,7 @@
 #include "Debug/DebugDrawService.h"
 #include "EnvQueryTypes.h"
 #include "EnvQueryDebugHelpers.h"
+#include "Components/PrimitiveComponent.h"
 #include "EQSRenderingComponent.generated.h"
 
 class IEQSQueryResultSourceInterface;
diff --git a/Engine/Source/Runtime/AIModule/Classes/EnvironmentQuery/Tests/EnvQueryTest_Distance.h b/Engine/Source/Runtime/AIModule/Classes/EnvironmentQuery/Tests/EnvQueryTest_Distance.h
--- a/Engine/Source/Runtime/AIModule/Classes/EnvironmentQuery/Tests/EnvQueryTest_Distance.h
+++ b/Engine/Source/Runtime/AIModule/Classes/EnvironmentQuery/Tests/EnvQueryTest_Distance.h
@@ -30,10 +30,9 @@ class UEnvQueryTest_Distance : public UEnvQueryTest
 	/** context */
 	UPROPERTY(EditDefaultsOnly, Category=Distance)
 	TSubclassOf<UEnvQueryContext> DistanceTo;
-#if WITH_DEV_AUTOMATION_TESTS
+
 	virtual void RunTest(FEnvQueryInstance& QueryInstance) const override;
 
 	virtual FText GetDescriptionTitle() const override;
 	virtual FText GetDescriptionDetails() const override;
-#endif
 };
diff --git a/Engine/Source/Runtime/AIModule/Classes/EnvironmentQuery/Tests/EnvQueryTest_Dot.h b/Engine/Source/Runtime/AIModule/Classes/EnvironmentQuery/Tests/EnvQueryTest_Dot.h
--- a/Engine/Source/Runtime/AIModule/Classes/EnvironmentQuery/Tests/EnvQueryTest_Dot.h
+++ b/Engine/Source/Runtime/AIModule/Classes/EnvironmentQuery/Tests/EnvQueryTest_Dot.h
@@ -39,7 +39,6 @@ protected:
 	UPROPERTY(EditDefaultsOnly, Category=Dot)
 	bool bAbsoluteValue;
 
-#if WITH_DEV_AUTOMATION_TESTS
 	virtual void RunTest(FEnvQueryInstance& QueryInstance) const override;
 
 	virtual FText GetDescriptionTitle() const override;
@@ -60,5 +59,4 @@ protected:
 
 	/** helper function: check if contexts are updated per item */
 	bool RequiresPerItemUpdates(TSubclassOf<UEnvQueryContext> LineFrom, TSubclassOf<UEnvQueryContext> LineTo, TSubclassOf<UEnvQueryContext> LineDirection, bool bUseDirectionContext) const;
-#endif
 };
diff --git a/Engine/Source/Runtime/AIModule/Classes/EnvironmentQuery/Tests/EnvQueryTest_GameplayTags.h b/Engine/Source/Runtime/AIModule/Classes/EnvironmentQuery/Tests/EnvQueryTest_GameplayTags.h
--- a/Engine/Source/Runtime/AIModule/Classes/EnvironmentQuery/Tests/EnvQueryTest_GameplayTags.h
+++ b/Engine/Source/Runtime/AIModule/Classes/EnvironmentQuery/Tests/EnvQueryTest_GameplayTags.h
@@ -13,13 +13,11 @@ class UEnvQueryTest_GameplayTags : public UEnvQueryTest
 	GENERATED_UCLASS_BODY()
 
 protected:
-#if WITH_DEV_AUTOMATION_TESTS	
 	virtual void RunTest(FEnvQueryInstance& QueryInstance) const override;
 
 	virtual FText GetDescriptionDetails() const override;
 
 	bool SatisfiesTest(IGameplayTagAssetInterface* ItemGameplayTagAssetInterface) const;
-#endif
 
 	UPROPERTY(EditAnywhere, Category=GameplayTagCheck)
 	EGameplayContainerMatchType TagsToMatch;
diff --git a/Engine/Source/Runtime/AIModule/Classes/EnvironmentQuery/Tests/EnvQueryTest_Trace.h b/Engine/Source/Runtime/AIModule/Classes/EnvironmentQuery/Tests/EnvQueryTest_Trace.h
--- a/Engine/Source/Runtime/AIModule/Classes/EnvironmentQuery/Tests/EnvQueryTest_Trace.h
+++ b/Engine/Source/Runtime/AIModule/Classes/EnvironmentQuery/Tests/EnvQueryTest_Trace.h
@@ -31,14 +31,12 @@ class UEnvQueryTest_Trace : public UEnvQueryTest
 	UPROPERTY(EditDefaultsOnly, Category=Trace)
 	TSubclassOf<UEnvQueryContext> Context;
 
-#if WITH_DEV_AUTOMATION_TESTS
 	virtual void RunTest(FEnvQueryInstance& QueryInstance) const override;
 
 	virtual FText GetDescriptionTitle() const override;
 	virtual FText GetDescriptionDetails() const override;
 
 	virtual void PostLoad() override;
-#endif //#if WITH_DEV_AUTOMATION_TESTS
 
 protected:
 
diff --git a/Engine/Source/Runtime/AIModule/Classes/Perception/AIPerceptionComponent.h b/Engine/Source/Runtime/AIModule/Classes/Perception/AIPerceptionComponent.h
--- a/Engine/Source/Runtime/AIModule/Classes/Perception/AIPerceptionComponent.h
+++ b/Engine/Source/Runtime/AIModule/Classes/Perception/AIPerceptionComponent.h
@@ -5,17 +5,14 @@
 #include "AISystem.h"
 #include "Perception/AISense.h"
 #include "AIPerceptionSystem.h"
-#if !UE_BUILD_SHIPPING
-#	include "Debug/DebugDrawService.h"
-#	include "Debug/GameplayDebuggerBaseObject.h"
-#endif
 #include "AIPerceptionComponent.generated.h"
 
 class AAIController;
-struct FVisualLogEntry;
 class UCanvas;
 class UAIPerceptionSystem;
 class UAISenseConfig;
+class FGameplayDebuggerCategory;
+struct FVisualLogEntry;
 
 DECLARE_DYNAMIC_MULTICAST_DELEGATE_OneParam(FPerceptionUpdatedDelegate, TArray<AActor*>, UpdatedActors);
 DECLARE_DYNAMIC_MULTICAST_DELEGATE_TwoParams(FActorPerceptionUpdatedDelegate, AActor*, Actor, FAIStimulus, Stimulus);
@@ -235,9 +232,9 @@ public:
 	bool HasAnyActiveStimulus(const AActor& Source) const;
 	bool HasActiveStimulus(const AActor& Source, FAISenseID Sense) const;
 
-#if !UE_BUILD_SHIPPING
-	void GrabGameplayDebuggerData(TArray<FString>& OnScreenStrings, TArray<FGameplayDebuggerShapeElement>& DebugShapes) const;
-#endif // !UE_BUILD_SHIPPING
+#if WITH_GAMEPLAY_DEBUGGER
+	virtual void DescribeSelfToGameplayDebugger(FGameplayDebuggerCategory* DebuggerCategory) const;
+#endif // WITH_GAMEPLAY_DEBUGGER
 
 #if ENABLE_VISUAL_LOG
 	virtual void DescribeSelfToVisLog(FVisualLogEntry* Snapshot) const;
diff --git a/Engine/Source/Runtime/AIModule/Classes/Perception/AIPerceptionSystem.h b/Engine/Source/Runtime/AIModule/Classes/Perception/AIPerceptionSystem.h
--- a/Engine/Source/Runtime/AIModule/Classes/Perception/AIPerceptionSystem.h
+++ b/Engine/Source/Runtime/AIModule/Classes/Perception/AIPerceptionSystem.h
@@ -187,21 +187,6 @@ private:
 	/** cached world's timestamp */
 	float CurrentTime;
 
-public:
-
-#if !UE_BUILD_SHIPPING
-	//----------------------------------------------------------------------//
-	// DEBUG
-	//----------------------------------------------------------------------//
-	TArray<FColor> DebugSenseColors;
-	FString PerceptionDebugLegend;	// describes which color means what
-
-	FColor GetSenseDebugColor(FAISenseID SenseID) const;
-	FString GetSenseName(FAISenseID SenseID) const;
-	FString GetPerceptionDebugLegend() const { return PerceptionDebugLegend; }
-#endif 
-
-private:
 	FTimerHandle AgeStimuliTimerHandle;
 };
 
