Fix build on a case-sensitive fs (AISense, GameplayDebuggingComponent).

(CLs 2885755, 5885807 in Dev-Platform).

diff --git a/Engine/Source/Developer/GameplayDebugger/Private/GameplayDebuggingComponent.cpp b/Engine/Source/Developer/GameplayDebugger/Private/GameplayDebuggingComponent.cpp
--- a/Engine/Source/Developer/GameplayDebugger/Private/GameplayDebuggingComponent.cpp
+++ b/Engine/Source/Developer/GameplayDebugger/Private/GameplayDebuggingComponent.cpp
@@ -33,7 +33,7 @@
 #endif // WITH_EDITOR
 
 #include "GameplayDebuggingComponent.h"
-#include "GAmeplayDebuggingControllerComponent.h"
+#include "GameplayDebuggingControllerComponent.h"
 #include "GameplayDebuggingReplicator.h"
 
 DEFINE_LOG_CATEGORY(LogGameplayDebugger);
diff --git a/Engine/Source/Runtime/AIModule/Private/Perception/AISense.cpp b/Engine/Source/Runtime/AIModule/Private/Perception/AISense.cpp
--- a/Engine/Source/Runtime/AIModule/Private/Perception/AISense.cpp
+++ b/Engine/Source/Runtime/AIModule/Private/Perception/AISense.cpp
@@ -8,7 +8,7 @@
 #include "Perception/AISenseConfig_Sight.h"
 #include "Perception/AISenseConfig_Team.h"
 #include "Perception/AISenseConfig_Touch.h"
-#include "PErception/AISense_Blueprint.h"
+#include "Perception/AISense_Blueprint.h"
 #include "Perception/AISense_Prediction.h"
 #include "Perception/AISense_Touch.h"
 
