Fix Mac non-unity

diff --git a/Engine/Source/Developer/GameplayDebugger/Private/GameplayDebuggerPlayerManager.cpp b/Engine/Source/Developer/GameplayDebugger/Private/GameplayDebuggerPlayerManager.cpp
--- a/Engine/Source/Developer/GameplayDebugger/Private/GameplayDebuggerPlayerManager.cpp
+++ b/Engine/Source/Developer/GameplayDebugger/Private/GameplayDebuggerPlayerManager.cpp
@@ -5,6 +5,7 @@
 #include "GameplayDebuggerAddonManager.h"
 #include "GameplayDebuggerCategoryReplicator.h"
 #include "GameplayDebuggerLocalController.h"
+#include "Engine/DebugCameraController.h"
 #include "Components/InputComponent.h"
 
 AGameplayDebuggerPlayerManager::AGameplayDebuggerPlayerManager(const FObjectInitializer& ObjectInitializer) : Super(ObjectInitializer)
@@ -102,7 +103,7 @@ void AGameplayDebuggerPlayerManager::UpdateAuthReplicators()
 	for (FConstPlayerControllerIterator It = World->GetPlayerControllerIterator(); It; It++)
 	{
 		APlayerController* TestPC = *It;
-		if (TestPC && !TestPC->IsA(ADebugCameraController::StaticClass()))
+		if (TestPC && !TestPC->IsA<ADebugCameraController>())
 		{
 			const bool bNeedsReplicator = (GetReplicator(*TestPC) == nullptr);
 			if (bNeedsReplicator)
