Copying //UE4/Dev-Framework to //UE4/Dev-Main (Source: //UE4/Dev-Framework @ 3182037)

diff --git a/Engine/Source/Developer/GameplayDebugger/Private/GameplayDebuggerLocalController.cpp b/Engine/Source/Developer/GameplayDebugger/Private/GameplayDebuggerLocalController.cpp
--- a/Engine/Source/Developer/GameplayDebugger/Private/GameplayDebuggerLocalController.cpp
+++ b/Engine/Source/Developer/GameplayDebugger/Private/GameplayDebuggerLocalController.cpp
@@ -13,6 +13,10 @@
 #include "GameFramework/PlayerInput.h"
 #include "GameFramework/Pawn.h"
 
+#if WITH_EDITOR
+#include "Editor/GameplayDebuggerEdMode.h"
+#endif // WITH_EDITOR
+
 UGameplayDebuggerLocalController::UGameplayDebuggerLocalController(const FObjectInitializer& ObjectInitializer) : Super(ObjectInitializer)
 {
 	bSimulateMode = false;
@@ -83,6 +87,10 @@ void UGameplayDebuggerLocalController::Cleanup()
 	if (bSimulateMode)
 	{
 		USelection::SelectionChangedEvent.RemoveAll(this);
+
+#if WITH_EDITOR
+		FGameplayDebuggerEdMode::SafeCloseMode();
+#endif // WITH_EDITOR
 	}
 
 	bNeedsCleanup = false;
@@ -105,14 +113,7 @@ void UGameplayDebuggerLocalController::OnDebugDraw(class UCanvas* Canvas, class
 		CanvasContext.CursorX = CanvasContext.DefaultX = PaddingLeft;
 		CanvasContext.CursorY = CanvasContext.DefaultY = PaddingTop;
 
-		if (bSimulateMode)
-		{
-			DrawSimulateHeader(CanvasContext);
-		}
-		else
-		{
-			DrawHeader(CanvasContext);
-		}
+		DrawHeader(CanvasContext);
 
 		if (DataPackMap.Num() != NumCategories)
 		{
@@ -142,9 +143,10 @@ void UGameplayDebuggerLocalController::DrawHeader(FGameplayDebuggerCanvasContext
 {
 	const int32 NumRows = (NumCategorySlots + 9) / 10;
 	const float LineHeight = CanvasContext.GetLineHeight();
-	const int32 NumExtensions = CachedReplicator->GetNumExtensions();
+	const int32 NumExtensions = bSimulateMode ? 0 : CachedReplicator->GetNumExtensions();
 	const int32 NumExtensionRows = (NumExtensions > 0) ? 1 : 0;
 	const float CanvasSizeX = CanvasContext.Canvas->SizeX - PaddingLeft - PaddingRight;
+	const float UsePaddingTop = PaddingTop + (bSimulateMode ? 30.0f : 0);
 	
 	const float BackgroundPadding = 5.0f;
 	const float BackgroundPaddingBothSides = BackgroundPadding * 2.0f;
@@ -159,35 +161,48 @@ void UGameplayDebuggerLocalController::DrawHeader(FGameplayDebuggerCanvasContext
 		ActiveRowTileItem.BlendMode = SE_BLEND_Translucent;
 		TileItemLower.BlendMode = SE_BLEND_Translucent;
 
-		CanvasContext.DrawItem(TileItemUpper, PaddingLeft - BackgroundPadding, PaddingTop - BackgroundPadding);
-		CanvasContext.DrawItem(ActiveRowTileItem, PaddingLeft - BackgroundPadding, PaddingTop - BackgroundPadding + TileItemUpper.Size.Y);
-		CanvasContext.DrawItem(TileItemLower, PaddingLeft - BackgroundPadding, PaddingTop - BackgroundPadding + TileItemUpper.Size.Y + ActiveRowTileItem.Size.Y);
+		CanvasContext.DrawItem(TileItemUpper, PaddingLeft - BackgroundPadding, UsePaddingTop - BackgroundPadding);
+		CanvasContext.DrawItem(ActiveRowTileItem, PaddingLeft - BackgroundPadding, UsePaddingTop - BackgroundPadding + TileItemUpper.Size.Y);
+		CanvasContext.DrawItem(TileItemLower, PaddingLeft - BackgroundPadding, UsePaddingTop - BackgroundPadding + TileItemUpper.Size.Y + ActiveRowTileItem.Size.Y);
 	}
 	else
 	{
 		FCanvasTileItem TileItem(FVector2D(0, 0), GWhiteTexture, FVector2D(CanvasSizeX, LineHeight * (NumRows + NumExtensionRows + 1)) + BackgroundPaddingBothSides, FLinearColor(0, 0, 0, 0.2f));
 		TileItem.BlendMode = SE_BLEND_Translucent;
-		CanvasContext.DrawItem(TileItem, PaddingLeft - BackgroundPadding, PaddingTop - BackgroundPadding);
+		CanvasContext.DrawItem(TileItem, PaddingLeft - BackgroundPadding, UsePaddingTop - BackgroundPadding);
 	}
 
-	CanvasContext.Printf(TEXT("Tap {yellow}%s{white} to close, use %s to toggle catories."), *ActivationKeyDesc, *CategoryKeysDesc);
+	CanvasContext.CursorY = UsePaddingTop;
+	if (bSimulateMode)
+	{
+		CanvasContext.Printf(TEXT("Clear {yellow}DebugAI{white} show flag to close, use %s to toggle catories."), *CategoryKeysDesc);
+
+		// reactivate editor mode when this is being drawn = show flag is set
+#if WITH_EDITOR
+		GLevelEditorModeTools().ActivateMode(FGameplayDebuggerEdMode::EM_GameplayDebugger);
+#endif // WITH_EDITOR
+	}
+	else
+	{
+		CanvasContext.Printf(TEXT("Tap {yellow}%s{white} to close, use %s to toggle catories."), *ActivationKeyDesc, *CategoryKeysDesc);
+	}
 
 	const FString DebugActorDesc = FString::Printf(TEXT("Debug actor: {cyan}%s"), *CachedReplicator->GetDebugActorName().ToString());
 	float DebugActorSizeX = 0.0f, DebugActorSizeY = 0.0f;
 	CanvasContext.MeasureString(DebugActorDesc, DebugActorSizeX, DebugActorSizeY);
-	CanvasContext.PrintAt(CanvasContext.Canvas->SizeX - PaddingRight - DebugActorSizeX, PaddingTop, DebugActorDesc);
+	CanvasContext.PrintAt(CanvasContext.Canvas->SizeX - PaddingRight - DebugActorSizeX, UsePaddingTop, DebugActorDesc);
 
 	const FString TimestampDesc = FString::Printf(TEXT("Time: %.2fs"), CachedReplicator->GetWorld()->GetTimeSeconds());
 	float TimestampSizeX = 0.0f, TimestampSizeY = 0.0f;
 	CanvasContext.MeasureString(TimestampDesc, TimestampSizeX, TimestampSizeY);
-	CanvasContext.PrintAt((CanvasSizeX - TimestampSizeX) * 0.5f, PaddingTop, TimestampDesc);
+	CanvasContext.PrintAt((CanvasSizeX - TimestampSizeX) * 0.5f, UsePaddingTop, TimestampDesc);
 
 	if (NumRows > 1)
 	{
 		const FString ChangeRowDesc = FString::Printf(TEXT("Prev row: {yellow}%s\n{white}Next row: {yellow}%s"), *RowUpKeyDesc, *RowDownKeyDesc);
 		float RowDescSizeX = 0.0f, RowDescSizeY = 0.0f;
 		CanvasContext.MeasureString(ChangeRowDesc, RowDescSizeX, RowDescSizeY);
-		CanvasContext.PrintAt(CanvasContext.Canvas->SizeX - PaddingRight - RowDescSizeX, PaddingTop + LineHeight * (NumExtensionRows + 1), ChangeRowDesc);
+		CanvasContext.PrintAt(CanvasContext.Canvas->SizeX - PaddingRight - RowDescSizeX, UsePaddingTop + LineHeight * (NumExtensionRows + 1), ChangeRowDesc);
 	}
 
 	if (NumExtensionRows)
@@ -251,31 +266,6 @@ void UGameplayDebuggerLocalController::DrawHeader(FGameplayDebuggerCanvasContext
 	CanvasContext.DefaultY = CanvasContext.CursorY + LineHeight;
 }
 
-void UGameplayDebuggerLocalController::DrawSimulateHeader(FGameplayDebuggerCanvasContext& CanvasContext)
-{
-	const float LineHeight = CanvasContext.GetLineHeight();
-	const float NewCategoryPosY = (CanvasContext.DefaultY * 2.0f) + (LineHeight * 3.0f);
-	const float BackgroundPadding = 5.0f;
-
-	const FString TimestampDesc = FString::Printf(TEXT("GameplayDebugger time: %.2fs"), CachedReplicator->GetWorld()->GetTimeSeconds());
-	float TimestampSizeX = 0.0f, TimestampSizeY = 0.0f;
-	CanvasContext.MeasureString(TimestampDesc, TimestampSizeX, TimestampSizeY);
-
-	const FString DebugActorDesc = FString::Printf(TEXT("Selected actor: {cyan}%s"), *CachedReplicator->GetDebugActorName().ToString());
-	float DebugActorSizeX = 0.0f, DebugActorSizeY = 0.0f;
-	CanvasContext.MeasureString(DebugActorDesc, DebugActorSizeX, DebugActorSizeY);
-
-	const float BackgroundSizeX = FMath::Max(DebugActorSizeX, TimestampSizeX) + (BackgroundPadding * 2.f);
-	FCanvasTileItem TileItem(FVector2D(0, 0), GWhiteTexture, FVector2D(BackgroundSizeX, (2.0f * LineHeight) + (2.0f * BackgroundPadding)), FLinearColor(0, 0, 0, 0.2f));
-	TileItem.BlendMode = SE_BLEND_Translucent;
-	CanvasContext.DrawItem(TileItem, (CanvasContext.Canvas->SizeX - BackgroundSizeX) * 0.5f, PaddingTop - BackgroundPadding);
-
-	CanvasContext.PrintAt((CanvasContext.Canvas->SizeX - TimestampSizeX) * 0.5f, PaddingTop, TimestampDesc);
-	CanvasContext.PrintAt((CanvasContext.Canvas->SizeX - DebugActorSizeX) * 0.5f, PaddingTop + LineHeight, DebugActorDesc);
-
-	CanvasContext.DefaultY = CanvasContext.CursorY = NewCategoryPosY;
-}
-
 void UGameplayDebuggerLocalController::DrawCategoryHeader(int32 CategoryId, TSharedRef<FGameplayDebuggerCategory> Category, FGameplayDebuggerCanvasContext& CanvasContext)
 {
 	FString DataPackDesc;
@@ -330,14 +320,17 @@ void UGameplayDebuggerLocalController::DrawCategoryHeader(int32 CategoryId, TSha
 
 void UGameplayDebuggerLocalController::BindInput(UInputComponent& InputComponent)
 {
-	TSet<FName> UsedBindings;
+	TSet<FName> NewBindings;
 
 	const UGameplayDebuggerConfig* SettingsCDO = UGameplayDebuggerConfig::StaticClass()->GetDefaultObject<UGameplayDebuggerConfig>();
-	InputComponent.BindKey(SettingsCDO->ActivationKey, IE_Pressed, this, &UGameplayDebuggerLocalController::OnActivationPressed);
-	InputComponent.BindKey(SettingsCDO->ActivationKey, IE_Released, this, &UGameplayDebuggerLocalController::OnActivationReleased);
-	UsedBindings.Add(SettingsCDO->ActivationKey.GetFName());
+	if (!bSimulateMode)
+	{
+		InputComponent.BindKey(SettingsCDO->ActivationKey, IE_Pressed, this, &UGameplayDebuggerLocalController::OnActivationPressed);
+		InputComponent.BindKey(SettingsCDO->ActivationKey, IE_Released, this, &UGameplayDebuggerLocalController::OnActivationReleased);
+		NewBindings.Add(SettingsCDO->ActivationKey.GetFName());
+	}
 
-	if (bIsLocallyEnabled)
+	if (bIsLocallyEnabled || bSimulateMode)
 	{
 		InputComponent.BindKey(SettingsCDO->CategorySlot0, IE_Pressed, this, &UGameplayDebuggerLocalController::OnCategory0Pressed);
 		InputComponent.BindKey(SettingsCDO->CategorySlot1, IE_Pressed, this, &UGameplayDebuggerLocalController::OnCategory1Pressed);
@@ -353,18 +346,18 @@ void UGameplayDebuggerLocalController::BindInput(UInputComponent& InputComponent
 		InputComponent.BindKey(SettingsCDO->CategoryRowPrevKey, IE_Pressed, this, &UGameplayDebuggerLocalController::OnCategoryRowUpPressed);
 		InputComponent.BindKey(SettingsCDO->CategoryRowNextKey, IE_Pressed, this, &UGameplayDebuggerLocalController::OnCategoryRowDownPressed);
 
-		UsedBindings.Add(SettingsCDO->CategorySlot0.GetFName());
-		UsedBindings.Add(SettingsCDO->CategorySlot1.GetFName());
-		UsedBindings.Add(SettingsCDO->CategorySlot2.GetFName());
-		UsedBindings.Add(SettingsCDO->CategorySlot3.GetFName());
-		UsedBindings.Add(SettingsCDO->CategorySlot4.GetFName());
-		UsedBindings.Add(SettingsCDO->CategorySlot5.GetFName());
-		UsedBindings.Add(SettingsCDO->CategorySlot6.GetFName());
-		UsedBindings.Add(SettingsCDO->CategorySlot7.GetFName());
-		UsedBindings.Add(SettingsCDO->CategorySlot8.GetFName());
-		UsedBindings.Add(SettingsCDO->CategorySlot9.GetFName());
-		UsedBindings.Add(SettingsCDO->CategoryRowPrevKey.GetFName());
-		UsedBindings.Add(SettingsCDO->CategoryRowNextKey.GetFName());
+		NewBindings.Add(SettingsCDO->CategorySlot0.GetFName());
+		NewBindings.Add(SettingsCDO->CategorySlot1.GetFName());
+		NewBindings.Add(SettingsCDO->CategorySlot2.GetFName());
+		NewBindings.Add(SettingsCDO->CategorySlot3.GetFName());
+		NewBindings.Add(SettingsCDO->CategorySlot4.GetFName());
+		NewBindings.Add(SettingsCDO->CategorySlot5.GetFName());
+		NewBindings.Add(SettingsCDO->CategorySlot6.GetFName());
+		NewBindings.Add(SettingsCDO->CategorySlot7.GetFName());
+		NewBindings.Add(SettingsCDO->CategorySlot8.GetFName());
+		NewBindings.Add(SettingsCDO->CategorySlot9.GetFName());
+		NewBindings.Add(SettingsCDO->CategoryRowPrevKey.GetFName());
+		NewBindings.Add(SettingsCDO->CategoryRowNextKey.GetFName());
 
 		for (int32 Idx = 0; Idx < NumCategories; Idx++)
 		{
@@ -377,16 +370,16 @@ void UGameplayDebuggerLocalController::BindInput(UInputComponent& InputComponent
 				if (HandlerData.Modifier.bPressed || HandlerData.Modifier.bReleased)
 				{
 					FInputChord InputChord(FKey(HandlerData.KeyName), HandlerData.Modifier.bShift, HandlerData.Modifier.bCtrl, HandlerData.Modifier.bAlt, HandlerData.Modifier.bCmd);
-					FInputKeyBinding NewBinding(InputChord, HandlerData.Modifier.bPressed ? IE_Pressed : IE_Released);
-					NewBinding.KeyDelegate.GetDelegateForManualSet().BindUObject(this, &UGameplayDebuggerLocalController::OnCategoryBindingEvent, Idx, HandlerIdx);
+					FInputKeyBinding InputBinding(InputChord, HandlerData.Modifier.bPressed ? IE_Pressed : IE_Released);
+					InputBinding.KeyDelegate.GetDelegateForManualSet().BindUObject(this, &UGameplayDebuggerLocalController::OnCategoryBindingEvent, Idx, HandlerIdx);
 
-					InputComponent.KeyBindings.Add(NewBinding);
-					UsedBindings.Add(HandlerData.KeyName);
+					InputComponent.KeyBindings.Add(InputBinding);
+					NewBindings.Add(HandlerData.KeyName);
 				}
 			}
 		}
 
-		const int32 NumExtentions = CachedReplicator->GetNumExtensions();
+		const int32 NumExtentions = bSimulateMode ? 0 : CachedReplicator->GetNumExtensions();
 		for (int32 Idx = 0; Idx < NumExtentions; Idx++)
 		{
 			TSharedRef<FGameplayDebuggerExtension> Extension = CachedReplicator->GetExtension(Idx);
@@ -398,11 +391,11 @@ void UGameplayDebuggerLocalController::BindInput(UInputComponent& InputComponent
 				if (HandlerData.Modifier.bPressed || HandlerData.Modifier.bReleased)
 				{
 					FInputChord InputChord(FKey(HandlerData.KeyName), HandlerData.Modifier.bShift, HandlerData.Modifier.bCtrl, HandlerData.Modifier.bAlt, HandlerData.Modifier.bCmd);
-					FInputKeyBinding NewBinding(InputChord, HandlerData.Modifier.bPressed ? IE_Pressed : IE_Released);
-					NewBinding.KeyDelegate.GetDelegateForManualSet().BindUObject(this, &UGameplayDebuggerLocalController::OnExtensionBindingEvent, Idx, HandlerIdx);
+					FInputKeyBinding InputBinding(InputChord, HandlerData.Modifier.bPressed ? IE_Pressed : IE_Released);
+					InputBinding.KeyDelegate.GetDelegateForManualSet().BindUObject(this, &UGameplayDebuggerLocalController::OnExtensionBindingEvent, Idx, HandlerIdx);
 
-					InputComponent.KeyBindings.Add(NewBinding);
-					UsedBindings.Add(HandlerData.KeyName);
+					InputComponent.KeyBindings.Add(InputBinding);
+					NewBindings.Add(HandlerData.KeyName);
 				}
 			}
 		}
@@ -410,8 +403,8 @@ void UGameplayDebuggerLocalController::BindInput(UInputComponent& InputComponent
 
 	if (CachedReplicator && CachedReplicator->GetReplicationOwner() && CachedReplicator->GetReplicationOwner()->PlayerInput)
 	{
-		TSet<FName> RemovedMasks = MaskedDebugBindings.Difference(UsedBindings);
-		TSet<FName> AddedMasks = UsedBindings.Difference(MaskedDebugBindings);
+		TSet<FName> RemovedMasks = UsedBindings.Difference(NewBindings);
+		TSet<FName> AddedMasks = NewBindings.Difference(UsedBindings);
 
 		UPlayerInput* Input = CachedReplicator->GetReplicationOwner()->PlayerInput;
 		for (int32 Idx = 0; Idx < Input->DebugExecBindings.Num(); Idx++)
@@ -426,10 +419,15 @@ void UGameplayDebuggerLocalController::BindInput(UInputComponent& InputComponent
 			}
 		}
 
-		MaskedDebugBindings = UsedBindings;
+		UsedBindings = NewBindings;
 	}
 }
 
+bool UGameplayDebuggerLocalController::IsKeyBound(const FName KeyName) const
+{
+	return UsedBindings.Contains(KeyName);
+}
+
 void UGameplayDebuggerLocalController::OnActivationPressed()
 {
 	bPrevLocallyEnabled = bIsLocallyEnabled;
@@ -667,6 +665,7 @@ void UGameplayDebuggerLocalController::OnSelectionChanged(UObject* Object)
 		}
 
 		CachedReplicator->SetDebugActor(SelectedPawn);
+		CachedReplicator->CollectCategoryData(/*bForce=*/true);
 	}
 }
 
diff --git a/Engine/Source/Developer/GameplayDebugger/Private/GameplayDebuggerModule.cpp b/Engine/Source/Developer/GameplayDebugger/Private/GameplayDebuggerModule.cpp
--- a/Engine/Source/Developer/GameplayDebugger/Private/GameplayDebuggerModule.cpp
+++ b/Engine/Source/Developer/GameplayDebugger/Private/GameplayDebuggerModule.cpp
@@ -12,13 +12,13 @@
 
 #if WITH_EDITOR
 #include "PropertyEditorModule.h"
+#include "EditorModeRegistry.h"
 #include "Editor/GameplayDebuggerCategoryConfigCustomization.h"
 #include "Editor/GameplayDebuggerExtensionConfigCustomization.h"
 #include "Editor/GameplayDebuggerInputConfigCustomization.h"
+#include "Editor/GameplayDebuggerEdMode.h"
 #endif
 
-#if !ENABLE_OLD_GAMEPLAY_DEBUGGER
-
 class FGameplayDebuggerModule : public IGameplayDebugger
 {
 public:
@@ -41,18 +41,7 @@ public:
 
 IMPLEMENT_MODULE(FGameplayDebuggerModule, GameplayDebugger)
 
-#else
-
-#include "GameplayDebuggerCompat.h"
-#define FGameplayDebuggerModule FGameplayDebuggerCompat
-
-#endif
-
-#if !ENABLE_OLD_GAMEPLAY_DEBUGGER
 void FGameplayDebuggerModule::StartupModule()
-#else
-void FGameplayDebuggerCompat::StartupNewDebugger()
-#endif
 {
 	// This code will execute after your module is loaded into memory (but after global variables are initialized, of course.)
 	FWorldDelegates::OnPostWorldInitialization.AddRaw(this, &FGameplayDebuggerModule::OnWorldInitialized);
@@ -74,6 +63,8 @@ void FGameplayDebuggerCompat::StartupNewDebugger()
 		PropertyEditorModule.RegisterCustomPropertyTypeLayout("GameplayDebuggerCategoryConfig", FOnGetPropertyTypeCustomizationInstance::CreateStatic(&FGameplayDebuggerCategoryConfigCustomization::MakeInstance));
 		PropertyEditorModule.RegisterCustomPropertyTypeLayout("GameplayDebuggerExtensionConfig", FOnGetPropertyTypeCustomizationInstance::CreateStatic(&FGameplayDebuggerExtensionConfigCustomization::MakeInstance));
 		PropertyEditorModule.RegisterCustomPropertyTypeLayout("GameplayDebuggerInputConfig", FOnGetPropertyTypeCustomizationInstance::CreateStatic(&FGameplayDebuggerInputConfigCustomization::MakeInstance));
+
+		FEditorModeRegistry::Get().RegisterMode<FGameplayDebuggerEdMode>(FGameplayDebuggerEdMode::EM_GameplayDebugger);
 #endif
 
 		AddonManager.RegisterExtension("GameHUD", FOnGetExtension::CreateStatic(&FGameplayDebuggerExtension_HUD::MakeInstance));
@@ -82,11 +73,7 @@ void FGameplayDebuggerCompat::StartupNewDebugger()
 	}
 }
 
-#if !ENABLE_OLD_GAMEPLAY_DEBUGGER
 void FGameplayDebuggerModule::ShutdownModule()
-#else
-void FGameplayDebuggerCompat::ShutdownNewDebugger()
-#endif
 {
 	// This function may be called during shutdown to clean up your module.  For modules that support dynamic reloading,
 	// we call this function before unloading the module.
@@ -103,6 +90,8 @@ void FGameplayDebuggerCompat::ShutdownNewDebugger()
 	PropertyEditorModule.UnregisterCustomPropertyTypeLayout("GameplayDebuggerCategoryConfig");
 	PropertyEditorModule.UnregisterCustomPropertyTypeLayout("GameplayDebuggerExtensionConfig");
 	PropertyEditorModule.UnregisterCustomPropertyTypeLayout("GameplayDebuggerInputConfig");
+
+	FEditorModeRegistry::Get().UnregisterMode(FGameplayDebuggerEdMode::EM_GameplayDebugger);
 #endif
 }
 
@@ -187,7 +176,3 @@ void FGameplayDebuggerModule::OnWorldInitialized(UWorld* World, const UWorld::In
 		GetPlayerManager(World);
 	}
 }
-
-#if ENABLE_OLD_GAMEPLAY_DEBUGGER
-#undef FGameplayDebuggerModule
-#endif // !ENABLE_OLD_GAMEPLAY_DEBUGGER
diff --git a/Engine/Source/Developer/GameplayDebugger/Private/GameplayDebuggerPlayerManager.cpp b/Engine/Source/Developer/GameplayDebugger/Private/GameplayDebuggerPlayerManager.cpp
--- a/Engine/Source/Developer/GameplayDebugger/Private/GameplayDebuggerPlayerManager.cpp
+++ b/Engine/Source/Developer/GameplayDebugger/Private/GameplayDebuggerPlayerManager.cpp
@@ -175,26 +175,30 @@ void AGameplayDebuggerPlayerManager::RefreshInputBindings(AGameplayDebuggerCateg
 
 AGameplayDebuggerCategoryReplicator* AGameplayDebuggerPlayerManager::GetReplicator(const APlayerController& OwnerPC) const
 {
-	for (int32 Idx = 0; Idx < PlayerData.Num(); Idx++)
-	{
-		AGameplayDebuggerCategoryReplicator* TestReplicator = PlayerData[Idx].Replicator;
-		if (TestReplicator && TestReplicator->GetReplicationOwner() == &OwnerPC)
-		{
-			return TestReplicator;
-		}
-	}
-
-	return nullptr;
+	const FGameplayDebuggerPlayerData* DataPtr = GetPlayerData(OwnerPC);
+	return DataPtr ? DataPtr->Replicator : nullptr;
 }
 
 UInputComponent* AGameplayDebuggerPlayerManager::GetInputComponent(const APlayerController& OwnerPC) const
+{
+	const FGameplayDebuggerPlayerData* DataPtr = GetPlayerData(OwnerPC);
+	return DataPtr ? DataPtr->InputComponent : nullptr;
+}
+
+UGameplayDebuggerLocalController* AGameplayDebuggerPlayerManager::GetLocalController(const APlayerController& OwnerPC) const
+{
+	const FGameplayDebuggerPlayerData* DataPtr = GetPlayerData(OwnerPC);
+	return DataPtr ? DataPtr->Controller : nullptr;
+}
+
+const FGameplayDebuggerPlayerData* AGameplayDebuggerPlayerManager::GetPlayerData(const APlayerController& OwnerPC) const
 {
 	for (int32 Idx = 0; Idx < PlayerData.Num(); Idx++)
 	{
 		const FGameplayDebuggerPlayerData& TestData = PlayerData[Idx];
 		if (TestData.Replicator && TestData.Replicator->GetReplicationOwner() == &OwnerPC)
 		{
-			return TestData.InputComponent;
+			return &TestData;
 		}
 	}
 
diff --git a/Engine/Source/Developer/GameplayDebugger/Public/GameplayDebugger.h b/Engine/Source/Developer/GameplayDebugger/Public/GameplayDebugger.h
--- a/Engine/Source/Developer/GameplayDebugger/Public/GameplayDebugger.h
+++ b/Engine/Source/Developer/GameplayDebugger/Public/GameplayDebugger.h
@@ -36,16 +36,6 @@
 //     PrivateDependencyModuleNames.Add("GameplayDebugger");
 //     Definitions.Add("WITH_GAMEPLAY_DEBUGGER=1");
 // }
-//
-//
-// -----------------------------------------------------------------------------------------------------
-// UPGRADE NOTES:
-// 
-// GameplayDebuggingComponent and GameplayDebuggingHUDComponent classes are no longer supported.
-// Instead of adding new project specific module, now you can create FGameplayDebuggerCategory classes
-// within project module. 
-// 
-// If you need to use old GameplayDebugger implementation, please set bEnableDeprecatedDebugger flag in GameplayDebugger.Build.cs
 // 
 
 #pragma once
@@ -94,31 +84,4 @@ public:
 	virtual void RegisterExtension(FName ExtensionName, FOnGetExtension MakeInstanceDelegate) = 0;
 	virtual void UnregisterExtension(FName ExtensionName) = 0;
 	virtual void NotifyExtensionsChanged() = 0;
-
-#if ENABLE_OLD_GAMEPLAY_DEBUGGER
-	virtual bool CreateGameplayDebuggerForPlayerController(APlayerController* PlayerController) = 0;
-	virtual bool IsGameplayDebuggerActiveForPlayerController(APlayerController* PlayerController) = 0;
-#else
-	DEPRECATED_FORGAME(4.13, "This function is now deprecated, please check GameplayDebugger.h for details.")
-	virtual bool CreateGameplayDebuggerForPlayerController(APlayerController* PlayerController) { return false; }
-	
-	DEPRECATED_FORGAME(4.13, "This function is now deprecated, please check GameplayDebugger.h for details.")
-	virtual bool IsGameplayDebuggerActiveForPlayerController(APlayerController* PlayerController) { return false; }
-#endif
-};
-
-#if ENABLE_OLD_GAMEPLAY_DEBUGGER
-class GameplayDebugger : public IGameplayDebugger
-{
-public:
-	static inline GameplayDebugger& Get() { return FModuleManager::LoadModuleChecked< GameplayDebugger >("GameplayDebugger"); }
-	virtual void UseNewGameplayDebugger() = 0;
-};
-#else
-class DEPRECATED_FORGAME(4.13, "This module interface is now deprecated, please check GameplayDebugger.h for details.") GameplayDebugger : public IGameplayDebugger
-{
-public:
-	static inline GameplayDebugger& Get() { return FModuleManager::LoadModuleChecked< GameplayDebugger >("GameplayDebugger"); }
-	virtual void UseNewGameplayDebugger() {}
 };
-#endif
diff --git a/Engine/Source/Developer/GameplayDebugger/Public/GameplayDebuggerLocalController.h b/Engine/Source/Developer/GameplayDebugger/Public/GameplayDebuggerLocalController.h
--- a/Engine/Source/Developer/GameplayDebugger/Public/GameplayDebuggerLocalController.h
+++ b/Engine/Source/Developer/GameplayDebugger/Public/GameplayDebuggerLocalController.h
@@ -28,6 +28,9 @@ class UGameplayDebuggerLocalController : public UObject
 	/** binds input actions */
 	void BindInput(UInputComponent& InputComponent);
 
+	/** checks if key is bound by any action */
+	bool IsKeyBound(const FName KeyName) const;
+
 protected:
 
 	UPROPERTY()
@@ -43,7 +46,7 @@ protected:
 	TArray<TArray<int32> > SlotCategoryIds;
 	TArray<FString> SlotNames;
 
-	TSet<FName> MaskedDebugBindings;
+	TSet<FName> UsedBindings;
 
 	uint32 bSimulateMode : 1;
 	uint32 bNeedsCleanup : 1;
@@ -97,9 +100,6 @@ protected:
 	/** draw header row */
 	void DrawHeader(FGameplayDebuggerCanvasContext& CanvasContext);
 
-	/** draw header row for simulate in editor mode */
-	void DrawSimulateHeader(FGameplayDebuggerCanvasContext& CanvasContext);
-
 	/** draw header for category */
 	void DrawCategoryHeader(int32 CategoryId, TSharedRef<FGameplayDebuggerCategory> Category, FGameplayDebuggerCanvasContext& CanvasContext);
 
diff --git a/Engine/Source/Developer/HierarchicalLODUtilities/Private/HierarchicalLODUtilities.cpp b/Engine/Source/Developer/HierarchicalLODUtilities/Private/HierarchicalLODUtilities.cpp
--- a/Engine/Source/Developer/HierarchicalLODUtilities/Private/HierarchicalLODUtilities.cpp
+++ b/Engine/Source/Developer/HierarchicalLODUtilities/Private/HierarchicalLODUtilities.cpp
@@ -79,28 +79,12 @@ void FHierarchicalLODUtilities::ExtractSubActorsFromLODActor(AActor* Actor, TArr
 
 float FHierarchicalLODUtilities::CalculateScreenSizeFromDrawDistance(const float SphereRadius, const FMatrix& ProjectionMatrix, const float Distance)
 {
-	// Only need one component from a view transformation; just calculate the one we're interested in.
-	const float Divisor = Distance;
-
-	// Get projection multiple accounting for view scaling.
-	const float ScreenMultiple = FMath::Max(1920.0f / 2.0f * ProjectionMatrix.M[0][0],
-		1080.0f / 2.0f * ProjectionMatrix.M[1][1]);
-
-	const float ScreenRadius = ScreenMultiple * SphereRadius / FMath::Max(Divisor, 1.0f);
-	const float ScreenArea = PI * ScreenRadius * ScreenRadius;
-	return FMath::Clamp(ScreenArea / (1920.0f * 1080.0f), 0.0f, 1.0f);
+	return ComputeBoundsScreenSize(FVector::ZeroVector, SphereRadius, FVector(0.0f, 0.0f, Distance), ProjectionMatrix);
 }
 
 float FHierarchicalLODUtilities::CalculateDrawDistanceFromScreenSize(const float SphereRadius, const float ScreenSize, const FMatrix& ProjectionMatrix)
 {
-	// Get projection multiple accounting for view scaling.
-	const float ScreenMultiple = FMath::Max(1920.0f / 2.0f * ProjectionMatrix.M[0][0],
-		1080.0f / 2.0f * ProjectionMatrix.M[1][1]);
-
-	// (ScreenMultiple * SphereRadius) / Sqrt(Screensize * 1920 * 1080.0f * PI) = Distance
-	const float Distance = (ScreenMultiple * SphereRadius) / FMath::Sqrt((ScreenSize * 1920.0f * 1080.0f) / PI);
-
-	return Distance;
+	return ComputeBoundsDrawDistance(ScreenSize, SphereRadius, ProjectionMatrix);
 }
 
 UPackage* FHierarchicalLODUtilities::CreateOrRetrieveLevelHLODPackage(ULevel* InLevel)
@@ -710,13 +694,13 @@ void FHierarchicalLODUtilities::DeleteLODActorsInHLODLevel(UWorld* InWorld, cons
 	}
 }
 
-int32 FHierarchicalLODUtilities::ComputeStaticMeshLODLevel(const TArray<FStaticMeshSourceModel>& SourceModels, const FStaticMeshRenderData* RenderData, const float ScreenAreaSize)
+int32 FHierarchicalLODUtilities::ComputeStaticMeshLODLevel(const TArray<FStaticMeshSourceModel>& SourceModels, const FStaticMeshRenderData* RenderData, const float ScreenSize)
 {	
 	const int32 NumLODs = SourceModels.Num();
 	// Walk backwards and return the first matching LOD
 	for (int32 LODIndex = NumLODs - 1; LODIndex >= 0; --LODIndex)
 	{
-		if (SourceModels[LODIndex].ScreenSize > ScreenAreaSize || ((SourceModels[LODIndex].ScreenSize == 0.0f) && (RenderData->ScreenSize[LODIndex] != SourceModels[LODIndex].ScreenSize) && (RenderData->ScreenSize[LODIndex] > ScreenAreaSize)))
+		if (SourceModels[LODIndex].ScreenSize > ScreenSize || ((SourceModels[LODIndex].ScreenSize == 0.0f) && (RenderData->ScreenSize[LODIndex] != SourceModels[LODIndex].ScreenSize) && (RenderData->ScreenSize[LODIndex] > ScreenSize)))
 		{
 			return FMath::Max(LODIndex, 0);
 		}
@@ -725,14 +709,14 @@ int32 FHierarchicalLODUtilities::ComputeStaticMeshLODLevel(const TArray<FStaticM
 	return 0;
 }
 
-int32 FHierarchicalLODUtilities::GetLODLevelForScreenAreaSize(const UStaticMeshComponent* StaticMeshComponent, const float ScreenAreaSize)
+int32 FHierarchicalLODUtilities::GetLODLevelForScreenSize(const UStaticMeshComponent* StaticMeshComponent, const float ScreenSize)
 {
 	check(StaticMeshComponent != nullptr);
 	const FStaticMeshRenderData* RenderData = StaticMeshComponent->GetStaticMesh()->RenderData.GetOwnedPointer();
 	checkf(RenderData != nullptr, TEXT("StaticMesh in StaticMeshComponent %s contains invalid render data"), *StaticMeshComponent->GetName());
 	checkf(StaticMeshComponent->GetStaticMesh()->SourceModels.Num() > 0, TEXT("StaticMesh in StaticMeshComponent %s contains no SourceModels"), *StaticMeshComponent->GetName());
 
-	return ComputeStaticMeshLODLevel(StaticMeshComponent->GetStaticMesh()->SourceModels, RenderData, ScreenAreaSize);
+	return ComputeStaticMeshLODLevel(StaticMeshComponent->GetStaticMesh()->SourceModels, RenderData, ScreenSize);
 }
 
 AHierarchicalLODVolume* FHierarchicalLODUtilities::CreateVolumeForLODActor(ALODActor* InLODActor, UWorld* InWorld)
diff --git a/Engine/Source/Developer/HierarchicalLODUtilities/Public/HierarchicalLODUtilities.h b/Engine/Source/Developer/HierarchicalLODUtilities/Public/HierarchicalLODUtilities.h
--- a/Engine/Source/Developer/HierarchicalLODUtilities/Public/HierarchicalLODUtilities.h
+++ b/Engine/Source/Developer/HierarchicalLODUtilities/Public/HierarchicalLODUtilities.h
@@ -35,8 +35,8 @@ public:
 	virtual void DestroyLODActor(ALODActor* InActor) override;
 	virtual void ExtractStaticMeshActorsFromLODActor(ALODActor* LODActor, TArray<AActor*> &InOutActors) override;
 	virtual void DeleteLODActorsInHLODLevel(UWorld* InWorld, const int32 HLODLevelIndex) override;
-	virtual int32 ComputeStaticMeshLODLevel(const TArray<FStaticMeshSourceModel>& SourceModels, const FStaticMeshRenderData* RenderData, const float ScreenAreaSize) override;
-	virtual int32 GetLODLevelForScreenAreaSize(const UStaticMeshComponent* StaticMeshComponent, const float ScreenAreaSize) override;
+	virtual int32 ComputeStaticMeshLODLevel(const TArray<FStaticMeshSourceModel>& SourceModels, const FStaticMeshRenderData* RenderData, const float ScreenSize) override;
+	virtual int32 GetLODLevelForScreenSize(const UStaticMeshComponent* StaticMeshComponent, const float ScreenSize) override;
 	virtual AHierarchicalLODVolume* CreateVolumeForLODActor(ALODActor* InLODActor, UWorld* InWorld) override;
 	virtual void HandleActorModified(AActor* InActor) override;
 };
diff --git a/Engine/Source/Developer/HierarchicalLODUtilities/Public/IHierarchicalLODUtilities.h b/Engine/Source/Developer/HierarchicalLODUtilities/Public/IHierarchicalLODUtilities.h
--- a/Engine/Source/Developer/HierarchicalLODUtilities/Public/IHierarchicalLODUtilities.h
+++ b/Engine/Source/Developer/HierarchicalLODUtilities/Public/IHierarchicalLODUtilities.h
@@ -123,11 +123,11 @@ public:
 	/** Deletes all the ALODActors with the given HLODLevelIndex inside off InWorld */
 	virtual void DeleteLODActorsInHLODLevel(UWorld* InWorld, const int32 HLODLevelIndex) = 0;
 
-	/** Computes which LOD level of a Mesh corresponds to the given Distance (calculates closest screensize with distance) */
-	virtual int32 ComputeStaticMeshLODLevel(const TArray<FStaticMeshSourceModel>& SourceModels, const FStaticMeshRenderData* RenderData, const float ScreenAreaSize) = 0;
+	/** Computes which LOD level of a Mesh corresponds to the given Distance (calculates closest ScreenSize with distance) */
+	virtual int32 ComputeStaticMeshLODLevel(const TArray<FStaticMeshSourceModel>& SourceModels, const FStaticMeshRenderData* RenderData, const float ScreenSize) = 0;
 
-	/** Computes the LODLevel for a StaticMeshComponent taking into account the ScreenArea */
-	virtual int32 GetLODLevelForScreenAreaSize(const UStaticMeshComponent* StaticMeshComponent, const float ScreenAreaSize) = 0;
+	/** Computes the LODLevel for a StaticMeshComponent taking into account the ScreenSize */
+	virtual int32 GetLODLevelForScreenSize(const UStaticMeshComponent* StaticMeshComponent, const float ScreenSize) = 0;
 	
 	/**
 	* Creates a HierarchicalLODVolume using the bounds of a given LODActor
diff --git a/Engine/Source/Developer/MeshUtilities/Private/MeshUtilities.cpp b/Engine/Source/Developer/MeshUtilities/Private/MeshUtilities.cpp
--- a/Engine/Source/Developer/MeshUtilities/Private/MeshUtilities.cpp
+++ b/Engine/Source/Developer/MeshUtilities/Private/MeshUtilities.cpp
@@ -228,7 +228,7 @@ private:
 		int32 UseLOD, // does not build all LODs but only use this LOD to create base mesh
 		TArray<UObject*>& OutAssetsToSync,
 		FVector& OutMergedActorLocation,
-		const float ScreenAreaSize,
+		const float ScreenSize,
 		bool bSilent = false) const override;
 
 	virtual void MergeStaticMeshComponents(
@@ -239,7 +239,7 @@ private:
 		const FString& InBasePackageName,
 		TArray<UObject*>& OutAssetsToSync,
 		FVector& OutMergedActorLocation,
-		const float ScreenAreaSize,
+		const float ScreenSize,
 		bool bSilent = false) const override;
 
 	virtual void CreateProxyMesh(const TArray<AActor*>& InActors, const struct FMeshProxySettings& InMeshProxySettings, UPackage* InOuter, const FString& InProxyBasePackageName, const FGuid InGuid, FCreateProxyDelegate InProxyCreatedDelegate, const bool bAllowAsync,
@@ -5885,7 +5885,7 @@ void FMeshUtilities::CreateProxyMesh(const TArray<AActor*>& InActors, const stru
 		FRawMesh* RawMesh = new FRawMesh();
 		FMemory::Memzero(RawMesh, sizeof(FRawMesh));
 		
-		const int32 ProxyMeshSourceLODLevel = InMeshProxySettings.bCalculateCorrectLODModel ? Utilities->GetLODLevelForScreenAreaSize(StaticMeshComponent, Utilities->CalculateScreenSizeFromDrawDistance(StaticMeshComponent->Bounds.SphereRadius, ProjectionMatrix, EstimatedDistance)) : 0;
+		const int32 ProxyMeshSourceLODLevel = InMeshProxySettings.bCalculateCorrectLODModel ? Utilities->GetLODLevelForScreenSize(StaticMeshComponent, Utilities->CalculateScreenSizeFromDrawDistance(StaticMeshComponent->Bounds.SphereRadius, ProjectionMatrix, EstimatedDistance)) : 0;
 		// Proxy meshes should always propagate vertex colours for material baking
 		static const bool bPropagateVertexColours = true;
 
@@ -7448,11 +7448,11 @@ void FMeshUtilities::MergeActors(
 
 	UWorld* World = SourceActors[0]->GetWorld();
 	checkf(World != nullptr, TEXT("Invalid world retrieved from Actor"));
-	const float ScreenAreaSize = TNumericLimits<float>::Max();
-	MergeStaticMeshComponents(ComponentsToMerge, World, InSettings, InOuter, InBasePackageName, OutAssetsToSync, OutMergedActorLocation, ScreenAreaSize, bSilent);
+	const float ScreenSize = TNumericLimits<float>::Max();
+	MergeStaticMeshComponents(ComponentsToMerge, World, InSettings, InOuter, InBasePackageName, OutAssetsToSync, OutMergedActorLocation, ScreenSize, bSilent);
 }
 
-void FMeshUtilities::MergeStaticMeshComponents(const TArray<UStaticMeshComponent*>& ComponentsToMerge, UWorld* World, const FMeshMergingSettings& InSettings, UPackage* InOuter, const FString& InBasePackageName, TArray<UObject*>& OutAssetsToSync, FVector& OutMergedActorLocation, const float ScreenAreaSize, bool bSilent /*= false*/) const
+void FMeshUtilities::MergeStaticMeshComponents(const TArray<UStaticMeshComponent*>& ComponentsToMerge, UWorld* World, const FMeshMergingSettings& InSettings, UPackage* InOuter, const FString& InBasePackageName, TArray<UObject*>& OutAssetsToSync, FVector& OutMergedActorLocation, const float ScreenSize, bool bSilent /*= false*/) const
 			{
 	FScopedSlowTask SlowTask(100.f, (LOCTEXT("MergeStaticMeshComponents_BuildingMesh", "Merging Static Mesh Components")));
 	SlowTask.MakeDialog();
@@ -7510,11 +7510,11 @@ void FMeshUtilities::MergeStaticMeshComponents(const TArray<UStaticMeshComponent
 
 		// LOD index will be overridden if the user has chosen to pick it according to the viewing distance
 		int32 CalculatedLODIndex = -1;
-		if (InSettings.LODSelectionType == EMeshLODSelectionType::CalculateLOD && ScreenAreaSize > 0.0f && ScreenAreaSize < 1.0f)
+		if (InSettings.LODSelectionType == EMeshLODSelectionType::CalculateLOD && ScreenSize > 0.0f && ScreenSize < 1.0f)
 		{
 			FHierarchicalLODUtilitiesModule& Module = FModuleManager::LoadModuleChecked<FHierarchicalLODUtilitiesModule>("HierarchicalLODUtilities");
 			IHierarchicalLODUtilities* Utilities = Module.GetUtilities();
-			CalculatedLODIndex = Utilities->GetLODLevelForScreenAreaSize(StaticMeshComponent, ScreenAreaSize);
+			CalculatedLODIndex = Utilities->GetLODLevelForScreenSize(StaticMeshComponent, ScreenSize);
 		}
 		SourceMeshes[MeshId].SourceStaticMesh = StaticMeshComponent->GetStaticMesh();
 
diff --git a/Engine/Source/Editor/AnimGraph/Classes/AnimGraphNode_Base.h b/Engine/Source/Editor/AnimGraph/Classes/AnimGraphNode_Base.h
--- a/Engine/Source/Editor/AnimGraph/Classes/AnimGraphNode_Base.h
+++ b/Engine/Source/Editor/AnimGraph/Classes/AnimGraphNode_Base.h
@@ -273,8 +273,8 @@ protected:
 	//
 	virtual FPoseLinkMappingRecord GetLinkIDLocation(const UScriptStruct* NodeType, UEdGraphPin* InputLinkPin);
 
-	//
-	virtual void GetPinAssociatedProperty(const UScriptStruct* NodeType, UEdGraphPin* InputPin, UProperty*& OutProperty, int32& OutIndex);
+	/** Get the property (and possibly array index) associated with the supplied pin */
+	virtual void GetPinAssociatedProperty(const UScriptStruct* NodeType, const UEdGraphPin* InputPin, UProperty*& OutProperty, int32& OutIndex) const;
 
 	// Allocates or reallocates pins
 	void InternalPinCreation(TArray<UEdGraphPin*>* OldPins);
diff --git a/Engine/Source/Editor/AnimGraph/Private/AnimGraphNode_Base.cpp b/Engine/Source/Editor/AnimGraph/Private/AnimGraphNode_Base.cpp
--- a/Engine/Source/Editor/AnimGraph/Private/AnimGraphNode_Base.cpp
+++ b/Engine/Source/Editor/AnimGraph/Private/AnimGraphNode_Base.cpp
@@ -151,7 +151,7 @@ FText UAnimGraphNode_Base::GetMenuCategory() const
 	return FText::FromString(GetNodeCategory());
 }
 
-void UAnimGraphNode_Base::GetPinAssociatedProperty(const UScriptStruct* NodeType, UEdGraphPin* InputPin, UProperty*& OutProperty, int32& OutIndex)
+void UAnimGraphNode_Base::GetPinAssociatedProperty(const UScriptStruct* NodeType, const UEdGraphPin* InputPin, UProperty*& OutProperty, int32& OutIndex) const
 {
 	OutProperty = nullptr;
 	OutIndex = INDEX_NONE;
diff --git a/Engine/Source/Editor/AnimGraph/Private/AnimGraphNode_ModifyCurve.cpp b/Engine/Source/Editor/AnimGraph/Private/AnimGraphNode_ModifyCurve.cpp
new file mode 100644
--- /dev/null
+++ b/Engine/Source/Editor/AnimGraph/Private/AnimGraphNode_ModifyCurve.cpp
@@ -0,0 +1,172 @@
+// Copyright 1998-2016 Epic Games, Inc. All Rights Reserved.
+
+#include "AnimGraphPrivatePCH.h"
+
+#include "GraphEditorActions.h"
+#include "ScopedTransaction.h"
+#include "CompilerResultsLog.h"
+#include "AnimGraphNode_ModifyCurve.h"
+#include "AnimationGraphSchema.h"
+#include "BlueprintActionDatabaseRegistrar.h"
+
+#define LOCTEXT_NAMESPACE "ModifyCurve"
+
+
+UAnimGraphNode_ModifyCurve::UAnimGraphNode_ModifyCurve() 
+{
+}
+
+FString UAnimGraphNode_ModifyCurve::GetNodeCategory() const
+{
+	return TEXT("Skeletal Controls");
+}
+
+FText UAnimGraphNode_ModifyCurve::GetTooltipText() const
+{
+	return GetNodeTitle(ENodeTitleType::ListView);
+}
+
+FText UAnimGraphNode_ModifyCurve::GetNodeTitle(ENodeTitleType::Type TitleType) const
+{
+	return LOCTEXT("AnimGraphNode_ModifyCurve_Title", "Modify Curve");
+}
+
+
+TArray<FName> UAnimGraphNode_ModifyCurve::GetCurvesToAdd() const
+{
+	TArray<FName> CurvesToAdd;
+
+	const FSmartNameMapping* Mapping = GetAnimBlueprint()->TargetSkeleton->GetSmartNameContainer(USkeleton::AnimCurveMappingName);
+	if (Mapping)
+	{
+		TArray<SmartName::UID_Type> UidList;
+		Mapping->FillUidArray(UidList);
+
+		// Iterate through all curves..
+		for (SmartName::UID_Type Uid : UidList)
+		{
+			FSmartName SmartName;
+			Mapping->FindSmartNameByUID(Uid, SmartName);
+
+			// Don't offer curves we already have
+			if (!Node.CurveNames.Contains(SmartName.DisplayName))
+			{
+				CurvesToAdd.Add(SmartName.DisplayName);
+			}
+		}
+	}
+
+	return CurvesToAdd;
+}
+
+void UAnimGraphNode_ModifyCurve::GetAddCurveMenuActions(FMenuBuilder& MenuBuilder) const
+{
+	TArray<FName> CurvesToAdd = GetCurvesToAdd();
+	for (FName CurveName : CurvesToAdd)
+	{
+		FUIAction Action = FUIAction(FExecuteAction::CreateUObject(this, &UAnimGraphNode_ModifyCurve::AddCurvePin, CurveName));
+		MenuBuilder.AddMenuEntry(FText::FromName(CurveName), FText::GetEmpty(), FSlateIcon(), Action);
+	}
+}
+
+void UAnimGraphNode_ModifyCurve::GetRemoveCurveMenuActions(FMenuBuilder& MenuBuilder) const
+{
+	for (FName CurveName : Node.CurveNames)
+	{
+		FUIAction Action = FUIAction(FExecuteAction::CreateUObject(this, &UAnimGraphNode_ModifyCurve::RemoveCurvePin, CurveName));
+		MenuBuilder.AddMenuEntry(FText::FromName(CurveName), FText::GetEmpty(), FSlateIcon(), Action);
+	}
+}
+
+
+void UAnimGraphNode_ModifyCurve::GetContextMenuActions(const FGraphNodeContextMenuBuilder& Context) const
+{
+	if (!Context.bIsDebugging)
+	{
+		Context.MenuBuilder->BeginSection("AnimGraphNodeModifyCurve", LOCTEXT("ModifyCurve", "Modify Curve"));
+
+		// Clicked pin
+		if (Context.Pin != NULL)
+		{
+			// Get proeprty from pin
+			UProperty* AssociatedProperty;
+			int32 ArrayIndex;
+			GetPinAssociatedProperty(GetFNodeType(), Context.Pin, /*out*/ AssociatedProperty, /*out*/ ArrayIndex);
+			FName PinPropertyName = AssociatedProperty->GetFName();
+
+			if (PinPropertyName  == GET_MEMBER_NAME_CHECKED(FAnimNode_ModifyCurve, CurveValues) && Context.Pin->Direction == EGPD_Input)
+			{
+				FString PinName = Context.Pin->PinFriendlyName.ToString();
+				FUIAction Action = FUIAction( FExecuteAction::CreateUObject(this, &UAnimGraphNode_ModifyCurve::RemoveCurvePin, FName(*PinName)) );
+				FText RemovePinLabelText = FText::Format(LOCTEXT("RemoveThisPin", "Remove This Curve Pin: {0}"), FText::FromString(PinName));
+				Context.MenuBuilder->AddMenuEntry(RemovePinLabelText, LOCTEXT("RemoveThisPinTooltip", "Remove this curve pin from this node"), FSlateIcon(), Action);
+			}
+		}
+
+		// If we have more curves to add, create submenu to offer them
+		if (GetCurvesToAdd().Num() > 0)
+		{
+			Context.MenuBuilder->AddSubMenu(
+				LOCTEXT("AddCurvePin", "Add Curve Pin"),
+				LOCTEXT("AddCurvePinTooltip", "Add a new pin to drive a curve"),
+				FNewMenuDelegate::CreateUObject(this, &UAnimGraphNode_ModifyCurve::GetAddCurveMenuActions));
+		}
+
+		// If we have curves to remove, create submenu to offer them
+		if (Node.CurveNames.Num() > 0)
+		{
+			Context.MenuBuilder->AddSubMenu(
+				LOCTEXT("RemoveCurvePin", "Remove Curve Pin"),
+				LOCTEXT("RemoveCurvePinTooltip", "Remove a pin driving a curve"),
+				FNewMenuDelegate::CreateUObject(this, &UAnimGraphNode_ModifyCurve::GetRemoveCurveMenuActions));
+		}
+	}
+
+	Context.MenuBuilder->EndSection();
+}
+
+void UAnimGraphNode_ModifyCurve::RemoveCurvePin(FName CurveName)
+{
+	// Make sure we have a curve pin with that name
+	int32 CurveIndex = Node.CurveNames.Find(CurveName);
+	if (CurveIndex != INDEX_NONE)
+	{
+		FScopedTransaction Transaction( LOCTEXT("RemoveCurvePinTrans", "Remove Curve Pin") );
+		Modify();
+
+		Node.RemoveCurve(CurveIndex);
+	
+		ReconstructNode();
+		FBlueprintEditorUtils::MarkBlueprintAsStructurallyModified(GetBlueprint());
+	}
+}
+
+void UAnimGraphNode_ModifyCurve::AddCurvePin(FName CurveName)
+{
+	// Make sure it doesn't already exist
+	int32 CurveIndex = Node.CurveNames.Find(CurveName);
+	if (CurveIndex == INDEX_NONE)
+	{
+		FScopedTransaction Transaction(LOCTEXT("AddCurvePinTrans", "Add Curve Pin"));
+		Modify();
+
+		Node.AddCurve(CurveName, 0.f);
+
+		ReconstructNode();
+		FBlueprintEditorUtils::MarkBlueprintAsStructurallyModified(GetBlueprint());
+	}
+}
+
+
+void UAnimGraphNode_ModifyCurve::CustomizePinData(UEdGraphPin* Pin, FName SourcePropertyName, int32 ArrayIndex) const
+{
+	if (SourcePropertyName == GET_MEMBER_NAME_CHECKED(FAnimNode_ModifyCurve, CurveValues))
+	{
+		if (Node.CurveNames.IsValidIndex(ArrayIndex))
+		{
+			Pin->PinFriendlyName = FText::FromName(Node.CurveNames[ArrayIndex]);
+		}
+	}
+}
+
+#undef LOCTEXT_NAMESPACE
diff --git a/Engine/Source/Editor/AnimationEditor/Private/AnimationEditorMode.cpp b/Engine/Source/Editor/AnimationEditor/Private/AnimationEditorMode.cpp
--- a/Engine/Source/Editor/AnimationEditor/Private/AnimationEditorMode.cpp
+++ b/Engine/Source/Editor/AnimationEditor/Private/AnimationEditorMode.cpp
@@ -22,7 +22,7 @@ FAnimationEditorMode::FAnimationEditorMode(TSharedRef<FWorkflowCentricApplicatio
 	TabFactories.RegisterFactory(PersonaModule.CreateAdvancedPreviewSceneTabFactory(InHostingApp, AnimationEditor->GetPersonaToolkit()->GetPreviewScene()));
 	TabFactories.RegisterFactory(PersonaModule.CreateAnimationAssetBrowserTabFactory(InHostingApp, AnimationEditor->GetPersonaToolkit(), FOnOpenNewAsset::CreateSP(&AnimationEditor.Get(), &FAnimationEditor::HandleOpenNewAsset), FOnAnimationSequenceBrowserCreated::CreateSP(&AnimationEditor.Get(), &FAnimationEditor::HandleAnimationSequenceBrowserCreated), true));
 	TabFactories.RegisterFactory(PersonaModule.CreateAssetDetailsTabFactory(InHostingApp, FOnGetAsset::CreateSP(&AnimationEditor.Get(), &FAnimationEditor::HandleGetAsset), FOnDetailsCreated()));
-	TabFactories.RegisterFactory(PersonaModule.CreateCurveViewerTabFactory(InHostingApp, InSkeletonTree->GetEditableSkeleton(), AnimationEditor->GetPersonaToolkit()->GetPreviewScene(), AnimationEditor->OnCurvesChanged, AnimationEditor->OnPostUndo));
+	TabFactories.RegisterFactory(PersonaModule.CreateCurveViewerTabFactory(InHostingApp, InSkeletonTree->GetEditableSkeleton(), AnimationEditor->GetPersonaToolkit()->GetPreviewScene(), AnimationEditor->OnCurvesChanged, AnimationEditor->OnPostUndo, FOnObjectsSelected::CreateSP(&AnimationEditor.Get(), &FAnimationEditor::HandleObjectsSelected)));
 	TabFactories.RegisterFactory(PersonaModule.CreateSkeletonSlotNamesTabFactory(InHostingApp, InSkeletonTree->GetEditableSkeleton(), AnimationEditor->OnPostUndo, FOnObjectSelected::CreateSP(&AnimationEditor.Get(), &FAnimationEditor::HandleObjectSelected)));
 
 	TabLayout = FTabManager::NewLayout("Standalone_AnimationEditor_Layout_v1.1")
diff --git a/Engine/Source/Editor/BehaviorTreeEditor/Classes/EdGraphSchema_BehaviorTreeDecorator.h b/Engine/Source/Editor/BehaviorTreeEditor/Classes/EdGraphSchema_BehaviorTreeDecorator.h
--- a/Engine/Source/Editor/BehaviorTreeEditor/Classes/EdGraphSchema_BehaviorTreeDecorator.h
+++ b/Engine/Source/Editor/BehaviorTreeEditor/Classes/EdGraphSchema_BehaviorTreeDecorator.h
@@ -47,7 +47,6 @@ class UEdGraphSchema_BehaviorTreeDecorator : public UEdGraphSchema
 	UPROPERTY()
 	FString PC_Boolean;
 
-	void GetBreakLinkToSubMenuActions(class FMenuBuilder& MenuBuilder, class UEdGraphPin* InGraphPin);
 	void AddPin(class UEdGraphNode* InGraphNode);
 	void RemovePin(class UEdGraphPin* InGraphPin);
 
diff --git a/Engine/Source/Editor/BehaviorTreeEditor/Private/BehaviorTreeDebugger.cpp b/Engine/Source/Editor/BehaviorTreeEditor/Private/BehaviorTreeDebugger.cpp
--- a/Engine/Source/Editor/BehaviorTreeEditor/Private/BehaviorTreeDebugger.cpp
+++ b/Engine/Source/Editor/BehaviorTreeEditor/Private/BehaviorTreeDebugger.cpp
@@ -69,6 +69,7 @@ void FBehaviorTreeDebugger::Setup(UBehaviorTree* InTreeAsset, TSharedRef<FBehavi
 	TreeAsset = InTreeAsset;
 	DebuggerInstanceIndex = INDEX_NONE;
 	ActiveStepIndex = 0;
+	LastValidStepId = INDEX_NONE;
 	ActiveBreakpoints.Reset();
 	KnownInstances.Reset();
 
@@ -135,13 +136,26 @@ void FBehaviorTreeDebugger::Tick(float DeltaTime)
 #if USE_BEHAVIORTREE_DEBUGGER
 	TArray<uint16> EmptyPath;
 
+	int32 TestStepIndex = 0;
+	for (int32 Idx = TreeInstance->DebuggerSteps.Num() - 1; Idx >= 0; Idx--)
+	{
+		const FBehaviorTreeExecutionStep& Step = TreeInstance->DebuggerSteps[Idx];
+		if (Step.StepIndex == LastValidStepId)
+		{
+			TestStepIndex = Idx;
+			break;
+		}
+	}
+
 	// find index of previously displayed state and notify about all changes in between to give breakpoints a chance to trigger
-	for (int32 i = FMath::Max(0, ActiveStepIndex); i < TreeInstance->DebuggerSteps.Num(); i++)
+	for (int32 i = TestStepIndex; i < TreeInstance->DebuggerSteps.Num(); i++)
 	{
 		const FBehaviorTreeExecutionStep& Step = TreeInstance->DebuggerSteps[i];
 		if (Step.StepIndex > DisplayedStepIndex)
 		{
 			ActiveStepIndex = i;
+			LastValidStepId = Step.StepIndex;
+
 			UpdateDebuggerInstance();
 			UpdateAvailableActions();
 
@@ -294,8 +308,10 @@ void FBehaviorTreeDebugger::OnTreeStarted(const UBehaviorTreeComponent& OwnerCom
 	KnownInstances.AddUnique(KnownComp);
 }
 
-void FBehaviorTreeDebugger::ClearDebuggerState()
+void FBehaviorTreeDebugger::ClearDebuggerState(bool bKeepSubtree)
 {
+	LastValidStepId = bKeepSubtree ? LastValidStepId : INDEX_NONE;
+
 	DebuggerInstanceIndex = INDEX_NONE;
 	ActiveStepIndex = 0;
 	DisplayedStepIndex = INDEX_NONE;
@@ -1044,7 +1060,7 @@ void FBehaviorTreeDebugger::UpdateDebuggerViewOnInstanceChange()
 	}
 	else
 	{
-		ClearDebuggerState();
+		ClearDebuggerState(/*bKeepSubtreeData=*/true);
 	}
 #endif
 }
diff --git a/Engine/Source/Editor/BehaviorTreeEditor/Private/BehaviorTreeDebugger.h b/Engine/Source/Editor/BehaviorTreeEditor/Private/BehaviorTreeDebugger.h
--- a/Engine/Source/Editor/BehaviorTreeEditor/Private/BehaviorTreeDebugger.h
+++ b/Engine/Source/Editor/BehaviorTreeEditor/Private/BehaviorTreeDebugger.h
@@ -106,6 +106,9 @@ private:
 	/** index of displayed step, used to detect changes */
 	int32 DisplayedStepIndex;
 
+	/** id of last valid step */
+	int32 LastValidStepId;
+
 	/** indices of display steps for different step actions */
 	int32 StepForwardIntoIdx;
 	int32 StepForwardOverIdx;
@@ -140,7 +143,7 @@ private:
 	void UpdateDebuggerInstance();
 
 	/** clear all runtime variables */
-	void ClearDebuggerState();
+	void ClearDebuggerState(bool bKeepSubtree = false);
 
 	/** try using breakpoints on node change */
 	void OnActiveNodeChanged(const TArray<uint16>& ActivePath, const TArray<uint16>& PrevStepPath);
diff --git a/Engine/Source/Editor/BehaviorTreeEditor/Private/EdGraphSchema_BehaviorTreeDecorator.cpp b/Engine/Source/Editor/BehaviorTreeEditor/Private/EdGraphSchema_BehaviorTreeDecorator.cpp
--- a/Engine/Source/Editor/BehaviorTreeEditor/Private/EdGraphSchema_BehaviorTreeDecorator.cpp
+++ b/Engine/Source/Editor/BehaviorTreeEditor/Private/EdGraphSchema_BehaviorTreeDecorator.cpp
@@ -160,32 +160,13 @@ void UEdGraphSchema_BehaviorTreeDecorator::GetContextMenuActions(const UEdGraph*
 
 	if (InGraphPin)
 	{
-		MenuBuilder->BeginSection("DecoratorGraphSchemaPinActions", LOCTEXT("PinActionsMenuHeader", "Pin Actions"));
+		// Only display the 'Break Links' option if there is a link to break!
+		if (InGraphPin->LinkedTo.Num() > 0)
 		{
-			// Only display the 'Break Links' option if there is a link to break!
-			if (InGraphPin->LinkedTo.Num() > 0)
-			{
-				MenuBuilder->AddMenuEntry( FGraphEditorCommands::Get().BreakPinLinks );
-
-				// add sub menu for break link to
-				if(InGraphPin->LinkedTo.Num() > 1)
-				{
-					MenuBuilder->AddSubMenu(
-						LOCTEXT("BreakLinkTo", "Break Link To..." ),
-						LOCTEXT("BreakSpecificLinks", "Break a specific link..." ),
-						FNewMenuDelegate::CreateUObject( (UEdGraphSchema_BehaviorTreeDecorator*const)this, &UEdGraphSchema_BehaviorTreeDecorator::GetBreakLinkToSubMenuActions, const_cast<UEdGraphPin*>(InGraphPin)));
-				}
-				else
-				{
-					((UEdGraphSchema_BehaviorTreeDecorator*const)this)->GetBreakLinkToSubMenuActions(*MenuBuilder, const_cast<UEdGraphPin*>(InGraphPin));
-				}
-			}
-			else
-			{
-				MenuBuilder->AddMenuEntry(FGraphEditorCommands::Get().RemoveExecutionPin);
-			}
+			MenuBuilder->BeginSection("DecoratorGraphSchemaPinActions", LOCTEXT("PinActionsMenuHeader", "Pin Actions"));
+			MenuBuilder->AddMenuEntry( FGraphEditorCommands::Get().BreakPinLinks );
+			MenuBuilder->EndSection();
 		}
-		MenuBuilder->EndSection();
 	}
 	else if (InGraphNode)
 	{
@@ -199,50 +180,6 @@ void UEdGraphSchema_BehaviorTreeDecorator::GetContextMenuActions(const UEdGraph*
 	}
 }
 
-void UEdGraphSchema_BehaviorTreeDecorator::GetBreakLinkToSubMenuActions( class FMenuBuilder& MenuBuilder, UEdGraphPin* InGraphPin )
-{
-	// Make sure we have a unique name for every entry in the list
-	TMap< FString, uint32 > LinkTitleCount;
-
-	// Add all the links we could break from
-	for(TArray<class UEdGraphPin*>::TConstIterator Links(InGraphPin->LinkedTo); Links; ++Links)
-	{
-		UEdGraphPin* Pin = *Links;
-		FString TitleString = Pin->GetOwningNode()->GetNodeTitle(ENodeTitleType::ListView).ToString();
-		FText Title = FText::FromString( TitleString );
-		if ( Pin->PinName != TEXT("") )
-		{
-			TitleString = FString::Printf(TEXT("%s (%s)"), *TitleString, *Pin->PinName);
-
-			// Add name of connection if possible
-			FFormatNamedArguments Args;
-			Args.Add( TEXT("NodeTitle"), Title );
-			Args.Add( TEXT("PinName"), Pin->GetDisplayName() );
-			Title = FText::Format( LOCTEXT("BreakDescPin", "{NodeTitle} ({PinName})"), Args );
-		}
-
-		uint32 &Count = LinkTitleCount.FindOrAdd( TitleString );
-
-		FText Description;
-		FFormatNamedArguments Args;
-		Args.Add( TEXT("NodeTitle"), Title );
-		Args.Add( TEXT("NumberOfNodes"), Count );
-
-		if ( Count == 0 )
-		{
-			Description = FText::Format( LOCTEXT("BreakDesc", "Break link to {NodeTitle}"), Args );
-		}
-		else
-		{
-			Description = FText::Format( LOCTEXT("BreakDescMulti", "Break link to {NodeTitle} ({NumberOfNodes})"), Args );
-		}
-		++Count;
-
-		MenuBuilder.AddMenuEntry( Description, Description, FSlateIcon(), FUIAction(
-			FExecuteAction::CreateUObject((USoundClassGraphSchema*const)this, &USoundClassGraphSchema::BreakSinglePinLink, const_cast< UEdGraphPin* >(InGraphPin), *Links) ) );
-	}
-}
-
 const FPinConnectionResponse UEdGraphSchema_BehaviorTreeDecorator::CanCreateConnection(const UEdGraphPin* PinA, const UEdGraphPin* PinB) const
 {
 	if (PinA == NULL || PinB == NULL)
diff --git a/Engine/Source/Editor/BlueprintGraph/Classes/EdGraphSchema_K2.h b/Engine/Source/Editor/BlueprintGraph/Classes/EdGraphSchema_K2.h
--- a/Engine/Source/Editor/BlueprintGraph/Classes/EdGraphSchema_K2.h
+++ b/Engine/Source/Editor/BlueprintGraph/Classes/EdGraphSchema_K2.h
@@ -1022,6 +1022,9 @@ public:
 	 */
 	void LinkDataPinFromOutputToInput(UEdGraphNode* InOutputNode, UEdGraphNode* InInputNode) const;
 
+	/** Convert a deprecated node into a function call node, called from per-node ConvertDeprecatedNode */
+	UK2Node* ConvertDeprecatedNodeToFunctionCall(UK2Node* OldNode, UFunction* NewFunction, TMap<FString, FString>& OldPinToNewPinMap, UEdGraph* Graph) const;
+
 	/** some inherited schemas don't want anim-notify actions listed, so this is an easy way to check that */
 	virtual bool DoesSupportAnimNotifyActions() const { return true; }
 
diff --git a/Engine/Source/Editor/BlueprintGraph/Classes/K2Node.h b/Engine/Source/Editor/BlueprintGraph/Classes/K2Node.h
--- a/Engine/Source/Editor/BlueprintGraph/Classes/K2Node.h
+++ b/Engine/Source/Editor/BlueprintGraph/Classes/K2Node.h
@@ -146,6 +146,7 @@ class UK2Node : public UEdGraphNode
 
 	// UObject interface
 	BLUEPRINTGRAPH_API virtual void PostLoad() override;
+	BLUEPRINTGRAPH_API virtual void Serialize(FArchive& Ar) override;
 	// End of UObject interface
 
 	// UEdGraphNode interface
@@ -231,7 +232,12 @@ class UK2Node : public UEdGraphNode
 	 */
 	bool CreatePinsForFunctionEntryExit(const UFunction* Function, bool bForFunctionEntry);
 
+	/** Expands a node while compiling, which may add additional nodes or delete this node */
 	BLUEPRINTGRAPH_API virtual void ExpandNode(class FKismetCompilerContext& CompilerContext, UEdGraph* SourceGraph);
+
+	/** Performs a node-specific deprecation fixup, which may delete this node and replace it with another one */
+	BLUEPRINTGRAPH_API virtual void ConvertDeprecatedNode(UEdGraph* Graph, bool bOnlySafeChanges) {}
+
 	BLUEPRINTGRAPH_API virtual class FNodeHandlingFunctor* CreateNodeHandler(class FKismetCompilerContext& CompilerContext) const { return NULL; }
 
 	BLUEPRINTGRAPH_API void ExpandSplitPin(class FKismetCompilerContext* CompilerContext, UEdGraph* SourceGraph, UEdGraphPin* Pin);
diff --git a/Engine/Source/Editor/BlueprintGraph/Classes/K2Node_BreakStruct.h b/Engine/Source/Editor/BlueprintGraph/Classes/K2Node_BreakStruct.h
--- a/Engine/Source/Editor/BlueprintGraph/Classes/K2Node_BreakStruct.h
+++ b/Engine/Source/Editor/BlueprintGraph/Classes/K2Node_BreakStruct.h
@@ -52,6 +52,7 @@ class UK2Node_BreakStruct : public UK2Node_StructMemberGet
 	virtual class FNodeHandlingFunctor* CreateNodeHandler(class FKismetCompilerContext& CompilerContext) const override;
 	virtual void GetMenuActions(FBlueprintActionDatabaseRegistrar& ActionRegistrar) const override;
 	virtual FText GetMenuCategory() const override;
+	virtual void ConvertDeprecatedNode(UEdGraph* Graph, bool bOnlySafeChanges) override;
 	//~ End K2Node Interface
 
 private:
diff --git a/Engine/Source/Editor/BlueprintGraph/Classes/K2Node_MakeStruct.h b/Engine/Source/Editor/BlueprintGraph/Classes/K2Node_MakeStruct.h
--- a/Engine/Source/Editor/BlueprintGraph/Classes/K2Node_MakeStruct.h
+++ b/Engine/Source/Editor/BlueprintGraph/Classes/K2Node_MakeStruct.h
@@ -54,6 +54,7 @@ class UK2Node_MakeStruct : public UK2Node_StructMemberSet
 	virtual ERedirectType DoPinsMatchForReconstruction(const UEdGraphPin* NewPin, int32 NewPinIndex, const UEdGraphPin* OldPin, int32 OldPinIndex)  const override;
 	virtual void GetMenuActions(FBlueprintActionDatabaseRegistrar& ActionRegistrar) const override;
 	virtual FText GetMenuCategory() const override;
+	virtual void ConvertDeprecatedNode(UEdGraph* Graph, bool bOnlySafeChanges) override;
 	//~ End K2Node Interface
 
 protected:
diff --git a/Engine/Source/Editor/BlueprintGraph/Private/EdGraphSchema_K2.cpp b/Engine/Source/Editor/BlueprintGraph/Private/EdGraphSchema_K2.cpp
--- a/Engine/Source/Editor/BlueprintGraph/Private/EdGraphSchema_K2.cpp
+++ b/Engine/Source/Editor/BlueprintGraph/Private/EdGraphSchema_K2.cpp
@@ -5337,7 +5337,7 @@ struct FBackwardCompatibilityConversionHelper
 				if (ClassPin != Pin)
 				{
 					const auto OldPin = OldNode->FindPin(Pin->PinName);
-						if(NULL != OldPin)
+					if (NULL != OldPin)
 					{
 						OldPins.Add(OldPin);
 						if (!Schema.MovePinLinks(*OldPin, *Pin).CanSafeConnect())
@@ -5358,7 +5358,7 @@ struct FBackwardCompatibilityConversionHelper
 			OldNode->BreakAllNodeLinks();
 			for (auto PinIter = OldNode->Pins.CreateIterator(); PinIter; ++PinIter)
 			{
-					if(!OldPins.Contains(*PinIter))
+				if (!OldPins.Contains(*PinIter))
 				{
 					UEdGraphPin* Pin = *PinIter;
 					UE_LOG(LogBlueprint, Warning, TEXT("BackwardCompatibilityNodeConversion Error 'missing new pin' in blueprint: %s"),
@@ -5402,25 +5402,91 @@ struct FBackwardCompatibilityConversionHelper
 	{
 		if (ConversionParams.FuncScope)
 		{
-		const UFunction* OldFunc = ConversionParams.FuncScope->FindFunctionByName(ConversionParams.OldFuncName);
-		check(OldFunc);
-		const UFunction* NewFunc = ConversionParams.FuncScope->FindFunctionByName(ConversionParams.NewFuncName);
-		check(NewFunc);
+			const UFunction* OldFunc = ConversionParams.FuncScope->FindFunctionByName(ConversionParams.OldFuncName);
+			check(OldFunc);
+			const UFunction* NewFunc = ConversionParams.FuncScope->FindFunctionByName(ConversionParams.NewFuncName);
+			check(NewFunc);
 
-		for (auto It = Nodes.CreateIterator(); It; ++It)
-		{
-			if (OldFunc == (*It)->GetTargetFunction())
+			for (auto It = Nodes.CreateIterator(); It; ++It)
 			{
-				auto NewNode = NewObject<UK2Node_CallFunction>(Graph);
-				NewNode->SetFromFunction(NewFunc);
-				ConvertNode(*It, ConversionParams.BlueprintPinName, NewNode, 
-					ConversionParams.ClassPinName, Schema, bOnlyWithDefaultBlueprint);
+				if (OldFunc == (*It)->GetTargetFunction())
+				{
+					auto NewNode = NewObject<UK2Node_CallFunction>(Graph);
+					NewNode->SetFromFunction(NewFunc);
+					ConvertNode(*It, ConversionParams.BlueprintPinName, NewNode,
+						ConversionParams.ClassPinName, Schema, bOnlyWithDefaultBlueprint);
+				}
 			}
 		}
 	}
-	}
 };
 
+UK2Node* UEdGraphSchema_K2::ConvertDeprecatedNodeToFunctionCall(UK2Node* OldNode, UFunction* NewFunction, TMap<FString, FString>& OldPinToNewPinMap, UEdGraph* Graph) const
+{
+	const UBlueprint* Blueprint = OldNode->GetBlueprint();
+
+	UK2Node_CallFunction* CallFunctionNode = NewObject<UK2Node_CallFunction>(Graph);
+	check(CallFunctionNode);
+	CallFunctionNode->SetFlags(RF_Transactional);
+	Graph->AddNode(CallFunctionNode, false, false);
+	CallFunctionNode->SetFromFunction(NewFunction);
+	CallFunctionNode->CreateNewGuid();
+	CallFunctionNode->PostPlacedNewNode();
+	CallFunctionNode->AllocateDefaultPins();
+	CallFunctionNode->NodePosX = OldNode->NodePosX;
+	CallFunctionNode->NodePosY = OldNode->NodePosY;
+	bool bFailedToFindPin = false;
+	TArray<UEdGraphPin*> NewPinArray;
+
+	for (int32 PinIdx = 0; PinIdx < OldNode->Pins.Num(); ++PinIdx)
+	{
+		UEdGraphPin* OldPin = OldNode->Pins[PinIdx];
+
+		// Check to see if the pin name is mapped to a new one, if it is use it, otherwise just search for the pin under the old name
+		FString* NewPinNamePtr = OldPinToNewPinMap.Find(OldPin->PinName);
+		FString NewPinName = NewPinNamePtr ? *NewPinNamePtr : OldPin->PinName;
+
+		UEdGraphPin* NewPin = CallFunctionNode->FindPin(*NewPinName);
+
+		if (!NewPin)
+		{
+			bFailedToFindPin = true;
+
+			UE_LOG(LogBlueprint, Warning, TEXT("BackwardCompatibilityNodeConversion Error 'cannot find pin %s in native function %s' in blueprint: %s"),
+				*NewPinName,
+				*NewFunction->GetName(),
+				Blueprint ? *Blueprint->GetName() : TEXT("Unknown"));
+
+			break;
+		}
+
+		NewPinArray.Add(NewPin);
+	}
+
+	if (bFailedToFindPin)
+	{
+		// Failed, destroy node
+		CallFunctionNode->DestroyNode();
+		return nullptr;
+	}
+	
+	for (int32 PinIdx = 0; PinIdx < OldNode->Pins.Num(); ++PinIdx)
+	{
+		UEdGraphPin* OldPin = OldNode->Pins[PinIdx];
+		UEdGraphPin* NewPin = NewPinArray[PinIdx];
+
+		if (!Graph->GetSchema()->MovePinLinks(*OldPin, *NewPin).CanSafeConnect())
+		{
+			UE_LOG(LogBlueprint, Warning, TEXT("BackwardCompatibilityNodeConversion Error 'cannot safely move pin %s to %s' in blueprint: %s"),
+				*OldPin->PinName,
+				*NewPin->PinName,
+				Blueprint ? *Blueprint->GetName() : TEXT("Unknown"));
+		}
+	}
+	OldNode->DestroyNode();
+	return CallFunctionNode;
+}
+
 void UEdGraphSchema_K2::BackwardCompatibilityNodeConversion(UEdGraph* Graph, bool bOnlySafeChanges) const 
 {
 	if (Graph)
@@ -5458,173 +5524,13 @@ void UEdGraphSchema_K2::BackwardCompatibilityNodeConversion(UEdGraph* Graph, boo
 			}
 		}
 
-		/** Fix the old Make/Break Vector, Make/Break Vector 2D, and Make/Break Rotator nodes to use the native function call versions */
+		// Call per-node deprecation functions
+		TArray<UK2Node*> PossiblyDeprecatedNodes;
+		Graph->GetNodesOfClass<UK2Node>(PossiblyDeprecatedNodes);
 
-		TArray<UK2Node_MakeStruct*> MakeStructNodes;
-		Graph->GetNodesOfClass<UK2Node_MakeStruct>(MakeStructNodes);
-		for(auto It = MakeStructNodes.CreateIterator(); It; ++It)
+		for (UK2Node* Node : PossiblyDeprecatedNodes)
 		{
-			UK2Node_MakeStruct* OldMakeStructNode = *It;
-			check(NULL != OldMakeStructNode);
-
-			// user may have since deleted the struct type
-			if (OldMakeStructNode->StructType == NULL)
-			{
-				continue;
-			}
-
-			// Check to see if the struct has a native make/break that we should try to convert to.
-			if (OldMakeStructNode->StructType && OldMakeStructNode->StructType->HasMetaData(TEXT("HasNativeMake")))
-			{
-				UFunction* MakeNodeFunction = NULL;
-
-				// If any pins need to change their names during the conversion, add them to the map.
-				TMap<FString, FString> OldPinToNewPinMap;
-
-				if(OldMakeStructNode->StructType->GetName() == TEXT("Rotator"))
-				{
-					MakeNodeFunction = UKismetMathLibrary::StaticClass()->FindFunctionByName(GET_FUNCTION_NAME_CHECKED(UKismetMathLibrary, MakeRotator));
-					OldPinToNewPinMap.Add(TEXT("Rotator"), TEXT("ReturnValue"));
-				}
-				else if(OldMakeStructNode->StructType->GetName() == TEXT("Vector"))
-				{
-					MakeNodeFunction = UKismetMathLibrary::StaticClass()->FindFunctionByName(GET_FUNCTION_NAME_CHECKED(UKismetMathLibrary, MakeVector));
-					OldPinToNewPinMap.Add(TEXT("Vector"), TEXT("ReturnValue"));
-				}
-				else if(OldMakeStructNode->StructType->GetName() == TEXT("Vector2D"))
-				{
-					MakeNodeFunction = UKismetMathLibrary::StaticClass()->FindFunctionByName(GET_FUNCTION_NAME_CHECKED(UKismetMathLibrary, MakeVector2D));
-					OldPinToNewPinMap.Add(TEXT("Vector2D"), TEXT("ReturnValue"));
-				}
-
-				if(MakeNodeFunction)
-				{
-					UK2Node_CallFunction* CallFunctionNode = NewObject<UK2Node_CallFunction>(Graph);
-					check(CallFunctionNode);
-					CallFunctionNode->SetFlags(RF_Transactional);
-					Graph->AddNode(CallFunctionNode, false, false);
-					CallFunctionNode->SetFromFunction(MakeNodeFunction);
-					CallFunctionNode->CreateNewGuid();
-					CallFunctionNode->PostPlacedNewNode();
-					CallFunctionNode->AllocateDefaultPins();
-					CallFunctionNode->NodePosX = OldMakeStructNode->NodePosX;
-					CallFunctionNode->NodePosY = OldMakeStructNode->NodePosY;
-
-					for(int32 PinIdx = 0; PinIdx < OldMakeStructNode->Pins.Num(); ++PinIdx)
-					{
-						UEdGraphPin* OldPin = OldMakeStructNode->Pins[PinIdx];
-						UEdGraphPin* NewPin = NULL;
-
-						// Check to see if the pin name is mapped to a new one, if it is use it, otherwise just search for the pin under the old name
-						FString* NewPinNamePtr = OldPinToNewPinMap.Find(OldPin->PinName);
-						if(NewPinNamePtr)
-						{
-							NewPin = CallFunctionNode->FindPin(*NewPinNamePtr);
-						}
-						else
-						{
-							NewPin = CallFunctionNode->FindPin(OldPin->PinName);
-						}
-						check(NewPin);
-
-						if(!Graph->GetSchema()->MovePinLinks(*OldPin, *NewPin).CanSafeConnect())
-						{
-							const UBlueprint* Blueprint = OldMakeStructNode->GetBlueprint();
-							UE_LOG(LogBlueprint, Warning, TEXT("BackwardCompatibilityNodeConversion Error 'cannot safetly move pin %s to %s' in blueprint: %s"), 
-								*OldPin->PinName, 
-								*NewPin->PinName,
-								Blueprint ? *Blueprint->GetName() : TEXT("Unknown"));
-						}
-					}
-					OldMakeStructNode->DestroyNode();
-				}
-			}
-		}
-
-		TArray<UK2Node_BreakStruct*> BreakStructNodes;
-		Graph->GetNodesOfClass<UK2Node_BreakStruct>(BreakStructNodes);
-		for(auto It = BreakStructNodes.CreateIterator(); It; ++It)
-		{
-			UK2Node_BreakStruct* OldBreakStructNode = *It;
-			check(NULL != OldBreakStructNode);
-
-			// user may have since deleted the struct type
-			if (OldBreakStructNode->StructType == NULL)
-			{
-				continue;
-			}
-
-			// Check to see if the struct has a native make/break that we should try to convert to.
-			if (OldBreakStructNode->StructType && OldBreakStructNode->StructType->HasMetaData(TEXT("HasNativeBreak")))
-			{
-				UFunction* BreakNodeFunction = NULL;
-
-				// If any pins need to change their names during the conversion, add them to the map.
-				TMap<FString, FString> OldPinToNewPinMap;
-
-				if(OldBreakStructNode->StructType->GetName() == TEXT("Rotator"))
-				{
-					BreakNodeFunction = UKismetMathLibrary::StaticClass()->FindFunctionByName(GET_FUNCTION_NAME_CHECKED(UKismetMathLibrary, BreakRotator));
-					OldPinToNewPinMap.Add(TEXT("Rotator"), TEXT("InRot"));
-				}
-				else if(OldBreakStructNode->StructType->GetName() == TEXT("Vector"))
-				{
-					BreakNodeFunction = UKismetMathLibrary::StaticClass()->FindFunctionByName(GET_FUNCTION_NAME_CHECKED(UKismetMathLibrary, BreakVector));
-					OldPinToNewPinMap.Add(TEXT("Vector"), TEXT("InVec"));
-				}
-				else if(OldBreakStructNode->StructType->GetName() == TEXT("Vector2D"))
-				{
-					BreakNodeFunction = UKismetMathLibrary::StaticClass()->FindFunctionByName(GET_FUNCTION_NAME_CHECKED(UKismetMathLibrary, BreakVector2D));
-					OldPinToNewPinMap.Add(TEXT("Vector2D"), TEXT("InVec"));
-				}
-
-				if(BreakNodeFunction)
-				{
-					UK2Node_CallFunction* CallFunctionNode = NewObject<UK2Node_CallFunction>(Graph);
-					check(CallFunctionNode);
-					CallFunctionNode->SetFlags(RF_Transactional);
-					Graph->AddNode(CallFunctionNode, false, false);
-					CallFunctionNode->SetFromFunction(BreakNodeFunction);
-					CallFunctionNode->CreateNewGuid();
-					CallFunctionNode->PostPlacedNewNode();
-					CallFunctionNode->AllocateDefaultPins();
-					CallFunctionNode->NodePosX = OldBreakStructNode->NodePosX;
-					CallFunctionNode->NodePosY = OldBreakStructNode->NodePosY;
-
-					for(int32 PinIdx = 0; PinIdx < OldBreakStructNode->Pins.Num(); ++PinIdx)
-					{
-						UEdGraphPin* OldPin = OldBreakStructNode->Pins[PinIdx];
-						UEdGraphPin* NewPin = NULL;
-
-						// Check to see if the pin name is mapped to a new one, if it is use it, otherwise just search for the pin under the old name
-						FString* NewPinNamePtr = OldPinToNewPinMap.Find(OldPin->PinName);
-						if(NewPinNamePtr)
-						{
-							NewPin = CallFunctionNode->FindPin(*NewPinNamePtr);
-						}
-						else
-						{
-							NewPin = CallFunctionNode->FindPin(OldPin->PinName);
-						}
-						check(NewPin);
-
-						if(!Graph->GetSchema()->MovePinLinks(*OldPin, *NewPin).CanSafeConnect())
-						{
-							const UBlueprint* Blueprint = OldBreakStructNode->GetBlueprint();
-							UE_LOG(LogBlueprint, Warning, TEXT("BackwardCompatibilityNodeConversion Error 'cannot safetly move pin %s to %s' in blueprint: %s"), 
-								*OldPin->PinName, 
-								*NewPin->PinName,
-								Blueprint ? *Blueprint->GetName() : TEXT("Unknown"));
-						}
-					}
-					OldBreakStructNode->DestroyNode();
-				}
-			}
-		}
-
-		if (Graph->GetLinker() == nullptr)
-		{
-			return;
+			Node->ConvertDeprecatedNode(Graph, bOnlySafeChanges);
 		}
 	}
 }
diff --git a/Engine/Source/Editor/BlueprintGraph/Private/K2Node.cpp b/Engine/Source/Editor/BlueprintGraph/Private/K2Node.cpp
--- a/Engine/Source/Editor/BlueprintGraph/Private/K2Node.cpp
+++ b/Engine/Source/Editor/BlueprintGraph/Private/K2Node.cpp
@@ -61,6 +61,39 @@ void UK2Node::PostLoad()
 	FixupPinDefaultValues();
 }
 
+void UK2Node::Serialize(FArchive& Ar)
+{
+	Super::Serialize(Ar);
+
+	if (Ar.IsObjectReferenceCollector() && Ar.IsSaving())
+	{
+		// If looking for references during save, expand any default values on the pins
+		
+		for (const UEdGraphPin* Pin : Pins)
+		{
+			if (!Pin->bDefaultValueIsIgnored && !Pin->DefaultValue.IsEmpty() && Pin->PinType.PinCategory == UEdGraphSchema_K2::PC_Struct && Pin->PinType.PinSubCategoryObject.IsValid())
+			{
+				UScriptStruct* Struct = Cast<UScriptStruct>(Pin->PinType.PinSubCategoryObject.Get());
+
+				if (Struct)
+				{
+					int32 StructSize = Struct->GetStructureSize();
+					uint8* StructData = (uint8*)FMemory::Malloc(StructSize);
+					
+					// Import the literal text to a dummy struct and then serialize that
+					FOutputDeviceNull NullOutput;
+					Struct->InitializeStruct(StructData);
+					Struct->ImportText(*Pin->DefaultValue, StructData, nullptr, PPF_None, &NullOutput, Pin->PinName);
+					Struct->SerializeItem(Ar, StructData, nullptr);
+					Struct->DestroyStruct(StructData);
+
+					FMemory::Free(StructData);
+				}
+			}
+		}
+	}
+}
+
 void UK2Node::FixupPinDefaultValues()
 {
 	const int32 LinkerUE4Version = GetLinkerUE4Version();
diff --git a/Engine/Source/Editor/BlueprintGraph/Private/K2Node_BreakStruct.cpp b/Engine/Source/Editor/BlueprintGraph/Private/K2Node_BreakStruct.cpp
--- a/Engine/Source/Editor/BlueprintGraph/Private/K2Node_BreakStruct.cpp
+++ b/Engine/Source/Editor/BlueprintGraph/Private/K2Node_BreakStruct.cpp
@@ -8,6 +8,7 @@
 #include "BlueprintActionFilter.h"	// for FBlueprintActionContext
 #include "Editor/PropertyEditor/Public/PropertyCustomizationHelpers.h"
 #include "BlueprintEditorSettings.h"
+#include "Kismet/KismetMathLibrary.h"
 
 #define LOCTEXT_NAMESPACE "K2Node_BreakStruct"
 
@@ -428,4 +429,73 @@ void UK2Node_BreakStruct::Serialize(FArchive& Ar)
 	}
 }
 
+void UK2Node_BreakStruct::ConvertDeprecatedNode(UEdGraph* Graph, bool bOnlySafeChanges)
+{
+	const UEdGraphSchema_K2* Schema = GetDefault<UEdGraphSchema_K2>();
+
+	// User may have since deleted the struct type
+	if (StructType == nullptr)
+	{
+		return;
+	}
+
+	// Check to see if the struct has a native make/break that we should try to convert to.
+	if (StructType->HasMetaData(FBlueprintMetadata::MD_NativeBreakFunction))
+	{
+		UFunction* BreakNodeFunction = nullptr;
+
+		// If any pins need to change their names during the conversion, add them to the map.
+		TMap<FString, FString> OldPinToNewPinMap;
+
+		if (StructType == TBaseStructure<FRotator>::Get())
+		{
+			BreakNodeFunction = UKismetMathLibrary::StaticClass()->FindFunctionByName(GET_FUNCTION_NAME_CHECKED(UKismetMathLibrary, BreakRotator));
+			OldPinToNewPinMap.Add(TEXT("Rotator"), TEXT("InRot"));
+		}
+		else if (StructType == TBaseStructure<FVector>::Get())
+		{
+			BreakNodeFunction = UKismetMathLibrary::StaticClass()->FindFunctionByName(GET_FUNCTION_NAME_CHECKED(UKismetMathLibrary, BreakVector));
+			OldPinToNewPinMap.Add(TEXT("Vector"), TEXT("InVec"));
+		}
+		else if (StructType == TBaseStructure<FVector2D>::Get())
+		{
+			BreakNodeFunction = UKismetMathLibrary::StaticClass()->FindFunctionByName(GET_FUNCTION_NAME_CHECKED(UKismetMathLibrary, BreakVector2D));
+			OldPinToNewPinMap.Add(TEXT("Vector2D"), TEXT("InVec"));
+		}
+		else
+		{
+			const FString& MetaData = StructType->GetMetaData(FBlueprintMetadata::MD_NativeBreakFunction);
+			BreakNodeFunction = FindObject<UFunction>(nullptr, *MetaData, true);
+			if (BreakNodeFunction)
+			{
+				// Look for the first parameter
+				for (TFieldIterator<UProperty> FieldIterator(BreakNodeFunction); FieldIterator && (FieldIterator->PropertyFlags & CPF_Parm); ++FieldIterator)
+				{
+					if (FieldIterator->PropertyFlags & CPF_Parm && !(FieldIterator->PropertyFlags & CPF_ReturnParm))
+					{
+						OldPinToNewPinMap.Add(*StructType->GetName(), *FieldIterator->GetName());
+						break;
+					}
+				}
+
+				// If map is empty, didn't find parameter
+				if (OldPinToNewPinMap.Num() == 0)
+				{
+					const UBlueprint* Blueprint = GetBlueprint();
+					UE_LOG(LogBlueprint, Warning, TEXT("BackwardCompatibilityNodeConversion Error 'cannot find input pin for break node function %s in blueprint: %s"),
+						*BreakNodeFunction->GetName(),
+						Blueprint ? *Blueprint->GetName() : TEXT("Unknown"));
+
+					BreakNodeFunction = nullptr;
+				}
+			}
+		}
+
+		if (BreakNodeFunction)
+		{
+			Schema->ConvertDeprecatedNodeToFunctionCall(this, BreakNodeFunction, OldPinToNewPinMap, Graph);
+		}
+	}
+}
+
 #undef LOCTEXT_NAMESPACE
diff --git a/Engine/Source/Editor/BlueprintGraph/Private/K2Node_MakeStruct.cpp b/Engine/Source/Editor/BlueprintGraph/Private/K2Node_MakeStruct.cpp
--- a/Engine/Source/Editor/BlueprintGraph/Private/K2Node_MakeStruct.cpp
+++ b/Engine/Source/Editor/BlueprintGraph/Private/K2Node_MakeStruct.cpp
@@ -9,6 +9,7 @@
 #include "BlueprintActionDatabaseRegistrar.h"
 #include "BlueprintActionFilter.h"	// for FBlueprintActionContext
 #include "Editor/PropertyEditor/Public/PropertyCustomizationHelpers.h"
+#include "Kismet/KismetMathLibrary.h"
 
 #define LOCTEXT_NAMESPACE "K2Node_MakeStruct"
 
@@ -460,4 +461,56 @@ void UK2Node_MakeStruct::Serialize(FArchive& Ar)
 	}
 }
 
+void UK2Node_MakeStruct::ConvertDeprecatedNode(UEdGraph* Graph, bool bOnlySafeChanges)
+{
+	const UEdGraphSchema_K2* Schema = GetDefault<UEdGraphSchema_K2>();
+
+	// User may have since deleted the struct type
+	if (StructType == nullptr)
+	{
+		return;
+	}
+
+	// Check to see if the struct has a native make/break that we should try to convert to.
+	if (StructType->HasMetaData(FBlueprintMetadata::MD_NativeMakeFunction))
+	{
+		UFunction* MakeNodeFunction = nullptr;
+
+		// If any pins need to change their names during the conversion, add them to the map.
+		TMap<FString, FString> OldPinToNewPinMap;
+
+		if (StructType == TBaseStructure<FRotator>::Get())
+		{
+			MakeNodeFunction = UKismetMathLibrary::StaticClass()->FindFunctionByName(GET_FUNCTION_NAME_CHECKED(UKismetMathLibrary, MakeRotator));
+			OldPinToNewPinMap.Add(TEXT("Rotator"), TEXT("ReturnValue"));
+		}
+		else if (StructType == TBaseStructure<FVector>::Get())
+		{
+			MakeNodeFunction = UKismetMathLibrary::StaticClass()->FindFunctionByName(GET_FUNCTION_NAME_CHECKED(UKismetMathLibrary, MakeVector));
+			OldPinToNewPinMap.Add(TEXT("Vector"), TEXT("ReturnValue"));
+		}
+		else if (StructType == TBaseStructure<FVector2D>::Get())
+		{
+			MakeNodeFunction = UKismetMathLibrary::StaticClass()->FindFunctionByName(GET_FUNCTION_NAME_CHECKED(UKismetMathLibrary, MakeVector2D));
+			OldPinToNewPinMap.Add(TEXT("Vector2D"), TEXT("ReturnValue"));
+		}
+		else
+		{
+			const FString& MetaData = StructType->GetMetaData(FBlueprintMetadata::MD_NativeMakeFunction);
+			MakeNodeFunction = FindObject<UFunction>(nullptr, *MetaData, true);
+
+			if (MakeNodeFunction)
+			{
+				OldPinToNewPinMap.Add(*StructType->GetName(), TEXT("ReturnValue"));
+			}
+		}
+
+		if (MakeNodeFunction)
+		{
+			Schema->ConvertDeprecatedNodeToFunctionCall(this, MakeNodeFunction, OldPinToNewPinMap, Graph);
+		}
+	}
+}
+
+
 #undef LOCTEXT_NAMESPACE
diff --git a/Engine/Source/Editor/DetailCustomizations/Private/DetailCustomizations.cpp b/Engine/Source/Editor/DetailCustomizations/Private/DetailCustomizations.cpp
--- a/Engine/Source/Editor/DetailCustomizations/Private/DetailCustomizations.cpp
+++ b/Engine/Source/Editor/DetailCustomizations/Private/DetailCustomizations.cpp
@@ -85,6 +85,7 @@
 #include "RawDistributionVectorStructCustomization.h"
 #include "CollisionProfileNameCustomization.h"
 #include "DocumentationActorDetails.h"
+#include "SoundBaseDetails.h"
 #include "SoundWaveDetails.h"
 #include "AudioSettingsDetails.h"
 #include "DateTimeStructCustomization.h"
@@ -292,6 +293,7 @@ void FDetailCustomizationsModule::RegisterObjectCustomizations()
 	RegisterCustomClassLayout("PoseAsset", FOnGetDetailCustomizationInstance::CreateStatic(&FPoseAssetDetails::MakeInstance));
 	RegisterCustomClassLayout("AnimationAsset", FOnGetDetailCustomizationInstance::CreateStatic(&FAnimationAssetDetails::MakeInstance));
 
+	RegisterCustomClassLayout("SoundBase", FOnGetDetailCustomizationInstance::CreateStatic(&FSoundBaseDetails::MakeInstance));
 	RegisterCustomClassLayout("SoundWave", FOnGetDetailCustomizationInstance::CreateStatic(&FSoundWaveDetails::MakeInstance));
 	RegisterCustomClassLayout("DialogueWave", FOnGetDetailCustomizationInstance::CreateStatic(&FDialogueWaveDetails::MakeInstance));
 	RegisterCustomClassLayout("BodySetup", FOnGetDetailCustomizationInstance::CreateStatic(&FBodySetupDetails::MakeInstance));
diff --git a/Engine/Source/Editor/DetailCustomizations/Private/SoundBaseDetails.cpp b/Engine/Source/Editor/DetailCustomizations/Private/SoundBaseDetails.cpp
new file mode 100644
--- /dev/null
+++ b/Engine/Source/Editor/DetailCustomizations/Private/SoundBaseDetails.cpp
@@ -0,0 +1,22 @@
+// Copyright 1998-2016 Epic Games, Inc. All Rights Reserved.
+
+#include "DetailCustomizationsPrivatePCH.h"
+#include "SoundBaseDetails.h"
+#include "Sound/SoundBase.h"
+
+TSharedRef<IDetailCustomization> FSoundBaseDetails::MakeInstance()
+{
+	return MakeShareable(new FSoundBaseDetails);
+}
+
+void FSoundBaseDetails::CustomizeDetails(IDetailLayoutBuilder& DetailBuilder)
+{
+	if (!GetDefault<UEditorExperimentalSettings>()->bShowAudioMixerData)
+	{
+		TSharedRef<IPropertyHandle> Property = DetailBuilder.GetProperty("SoundSubmixObject", USoundBase::StaticClass());
+		Property->MarkHiddenByCustomization();
+
+		Property = DetailBuilder.GetProperty("SourceEffectChain", USoundBase::StaticClass());
+		Property->MarkHiddenByCustomization();
+	}
+}
diff --git a/Engine/Source/Editor/GameplayAbilitiesEditor/GameplayAbilitiesEditor.Build.cs b/Engine/Source/Editor/GameplayAbilitiesEditor/GameplayAbilitiesEditor.Build.cs
--- a/Engine/Source/Editor/GameplayAbilitiesEditor/GameplayAbilitiesEditor.Build.cs
+++ b/Engine/Source/Editor/GameplayAbilitiesEditor/GameplayAbilitiesEditor.Build.cs
@@ -4,21 +4,21 @@ namespace UnrealBuildTool.Rules
 {
 	public class GameplayAbilitiesEditor : ModuleRules
 	{
-        public GameplayAbilitiesEditor(TargetInfo Target)
+		public GameplayAbilitiesEditor(TargetInfo Target)
 		{
-            PrivateIncludePaths.AddRange(
-                new string[] {
-                    "Editor/GraphEditor/Private",
-				    "Editor/Kismet/Private",
+			PrivateIncludePaths.AddRange(
+				new string[] {
+					"Editor/GraphEditor/Private",
+					"Editor/Kismet/Private",
 					"Editor/GameplayAbilitiesEditor/Private",
-                    "Developer/AssetTools/Private",
+					"Developer/AssetTools/Private",
 				}
 			);
 
-            PublicDependencyModuleNames.Add("GameplayTasks");
+			PublicDependencyModuleNames.Add("GameplayTasks");
 
-            PrivateDependencyModuleNames.AddRange(
-                new string[]
+			PrivateDependencyModuleNames.AddRange(
+				new string[]
 				{
 					// ... add private dependencies that you statically link with here ...
 					"Core",
@@ -26,26 +26,34 @@ namespace UnrealBuildTool.Rules
 					"Engine",
 					"AssetTools",
 					"ClassViewer",
-                    "GameplayTags",
+					"GameplayTags",
+					"GameplayTagsEditor",
 					"GameplayAbilities",
-                    "GameplayTasksEditor",
-                    "InputCore",
-                    "PropertyEditor",
+					"GameplayTasksEditor",
+					"InputCore",
+					"PropertyEditor",
 					"Slate",
 					"SlateCore",
-                    "EditorStyle",
+					"EditorStyle",
 					"BlueprintGraph",
-                    "Kismet",
+					"Kismet",
 					"KismetCompiler",
 					"GraphEditor",
 					"MainFrame",
 					"UnrealEd",
-                    "WorkspaceMenuStructure",
-                    "ContentBrowser",
-                    "EditorWidgets",
-                    "SourceControl",
+					"WorkspaceMenuStructure",
+					"ContentBrowser",
+					"EditorWidgets",
+					"SourceControl",
 				}
 			);
+
+			CircularlyReferencedDependentModules.AddRange(
+			new string[]
+			{
+                "GameplayTagsEditor",
+			}
+		);
 		}
 	}
 }
diff --git a/Engine/Source/Editor/GameplayAbilitiesEditor/Private/GameplayAbilitiesEditorModule.cpp b/Engine/Source/Editor/GameplayAbilitiesEditor/Private/GameplayAbilitiesEditorModule.cpp
--- a/Engine/Source/Editor/GameplayAbilitiesEditor/Private/GameplayAbilitiesEditorModule.cpp
+++ b/Engine/Source/Editor/GameplayAbilitiesEditor/Private/GameplayAbilitiesEditorModule.cpp
@@ -16,12 +16,10 @@
 #include "IAssetTypeActions.h"
 #include "AssetToolsModule.h"
 #include "AssetTypeActions_GameplayAbilitiesBlueprint.h"
-#include "AssetTypeActions_GameplayEffect.h"
 #include "GameplayAbilitiesGraphPanelPinFactory.h"
-#include "GameplayAbilitiesGraphPanelNodeFactory.h"
 #include "GameplayCueTagDetails.h"
 
-#include "Runtime/GameplayTags/Public/GameplayTagsModule.h"
+#include "GameplayTagsModule.h"
 #include "Editor/BlueprintGraph/Public/BlueprintActionDatabase.h"
 #include "K2Node_GameplayCueEvent.h"
 
@@ -80,9 +78,6 @@ private:
 	/** Pin factory for abilities graph; Cached so it can be unregistered */
 	TSharedPtr<FGameplayAbilitiesGraphPanelPinFactory> GameplayAbilitiesGraphPanelPinFactory;
 
-	/** Node factory for abilities graph; Cached so it can be unregistered */
-	TSharedPtr<FGameplayAbilitiesGraphPanelNodeFactory> GameplayAbilitiesGraphPanelNodeFactory;
-
 	/** Handle to the registered GameplayTagTreeChanged delegate */
 	FDelegateHandle GameplayTagTreeChangedDelegateHandle;
 
@@ -128,20 +123,15 @@ void FGameplayAbilitiesEditorModule::StartupModule()
 	// Register asset types
 	IAssetTools& AssetTools = FModuleManager::LoadModuleChecked<FAssetToolsModule>("AssetTools").Get();
 	TSharedRef<IAssetTypeActions> GABAction = MakeShareable(new FAssetTypeActions_GameplayAbilitiesBlueprint());
-	TSharedRef<IAssetTypeActions> GEAction = MakeShareable(new FAssetTypeActions_GameplayEffect());
 	RegisterAssetTypeAction(AssetTools, GABAction);
-	RegisterAssetTypeAction(AssetTools, GEAction);
 
 	// Register factories for pins and nodes
 	GameplayAbilitiesGraphPanelPinFactory = MakeShareable(new FGameplayAbilitiesGraphPanelPinFactory());
 	FEdGraphUtilities::RegisterVisualPinFactory(GameplayAbilitiesGraphPanelPinFactory);
 
-	GameplayAbilitiesGraphPanelNodeFactory = MakeShareable(new FGameplayAbilitiesGraphPanelNodeFactory());
-	FEdGraphUtilities::RegisterVisualNodeFactory(GameplayAbilitiesGraphPanelNodeFactory);
-
 	// Listen for changes to the gameplay tag tree so we can refresh blueprint actions for the GameplayCueEvent node
-	UGameplayTagsManager& GameplayTagsManager = IGameplayTagsModule::GetGameplayTagsManager();
-	GameplayTagTreeChangedDelegateHandle = GameplayTagsManager.OnGameplayTagTreeChanged().AddStatic(&FGameplayAbilitiesEditorModule::GameplayTagTreeChanged);
+	UGameplayTagsManager& GameplayTagsManager = UGameplayTagsManager::Get();
+	GameplayTagTreeChangedDelegateHandle = IGameplayTagsModule::OnGameplayTagTreeChanged.AddStatic(&FGameplayAbilitiesEditorModule::GameplayTagTreeChanged);
 
 	// GameplayCue editor
 	FGlobalTabmanager::Get()->RegisterNomadTabSpawner( FName(TEXT("GameplayCueApp")), FOnSpawnTab::CreateRaw(this, &FGameplayAbilitiesEditorModule::SpawnGameplayCueEditorTab))
@@ -280,16 +270,9 @@ void FGameplayAbilitiesEditorModule::ShutdownModule()
 		GameplayAbilitiesGraphPanelPinFactory.Reset();
 	}
 
-	if (GameplayAbilitiesGraphPanelNodeFactory.IsValid())
-	{
-		FEdGraphUtilities::UnregisterVisualNodeFactory(GameplayAbilitiesGraphPanelNodeFactory);
-		GameplayAbilitiesGraphPanelNodeFactory.Reset();
-	}
-
 	if ( UObjectInitialized() && IGameplayTagsModule::IsAvailable() )
 	{
-		UGameplayTagsManager& GameplayTagsManager = IGameplayTagsModule::GetGameplayTagsManager();
-		GameplayTagsManager.OnGameplayTagTreeChanged().Remove(GameplayTagTreeChangedDelegateHandle);
+		IGameplayTagsModule::OnGameplayTagTreeChanged.Remove(GameplayTagTreeChangedDelegateHandle);
 	}
 }
 
diff --git a/Engine/Source/Editor/GameplayAbilitiesEditor/Private/GameplayAbilityGraphSchema.cpp b/Engine/Source/Editor/GameplayAbilitiesEditor/Private/GameplayAbilityGraphSchema.cpp
--- a/Engine/Source/Editor/GameplayAbilitiesEditor/Private/GameplayAbilityGraphSchema.cpp
+++ b/Engine/Source/Editor/GameplayAbilitiesEditor/Private/GameplayAbilityGraphSchema.cpp
@@ -1,7 +1,6 @@
 // Copyright 1998-2016 Epic Games, Inc. All Rights Reserved.
 
 #include "AbilitySystemEditorPrivatePCH.h"
-#include "K2Node_GameplayEffectVariable.h"
 #include "GameplayEffect.h"
 #include "GameplayAbilityGraphSchema.h"
 #include "Kismet/KismetSystemLibrary.h"
@@ -16,37 +15,6 @@ UGameplayAbilityGraphSchema::UGameplayAbilityGraphSchema(const FObjectInitialize
 
 UK2Node_VariableGet* UGameplayAbilityGraphSchema::SpawnVariableGetNode(const FVector2D GraphPosition, class UEdGraph* ParentGraph, FName VariableName, UStruct* Source) const
 {
-	// Perform handling to create custom nodes for some classes
-	UProperty* VarProp = FindField<UProperty>(Source, VariableName);
-	UObjectProperty* ObjProp = Cast<UObjectProperty>(VarProp);
-	if (ObjProp)
-	{
-		UBlueprint* Blueprint = FBlueprintEditorUtils::FindBlueprintForGraph(ParentGraph);
-		TSubclassOf<UObject> GenClass = Blueprint->GeneratedClass;
-
-		UObject* ActiveObject = GenClass->GetDefaultObject();
-
-		// If the variable is a GameplayEffect create a custom node to show it
-		FString PropType;
-		ObjProp->GetCPPMacroType(PropType);
-		if (PropType == "UGameplayEffect")
-		{
-			UK2Node_GameplayEffectVariable* NodeTemplate = NewObject<UK2Node_GameplayEffectVariable>();
-			UEdGraphSchema_K2::ConfigureVarNode(NodeTemplate, VariableName, Source, Blueprint);
-
-			UK2Node_GameplayEffectVariable* VariableNode = FEdGraphSchemaAction_K2NewNode::SpawnNodeFromTemplate<UK2Node_GameplayEffectVariable>(ParentGraph, NodeTemplate, GraphPosition);
-
-			UGameplayEffect* GameplayEffect = Cast<UGameplayEffect>(ObjProp->GetObjectPropertyValue_InContainer(ActiveObject));
-			if (GameplayEffect)
-			{
-				VariableNode->GameplayEffect = GameplayEffect;
-			}
-
-			return VariableNode;
-		}
-	}
-
-	// Couldn't find an appropriate custom node for this variable, use the generic case
 	return Super::SpawnVariableGetNode(GraphPosition, ParentGraph, VariableName, Source);
 }
 
diff --git a/Engine/Source/Editor/GameplayAbilitiesEditor/Private/K2Node_GameplayCueEvent.cpp b/Engine/Source/Editor/GameplayAbilitiesEditor/Private/K2Node_GameplayCueEvent.cpp
--- a/Engine/Source/Editor/GameplayAbilitiesEditor/Private/K2Node_GameplayCueEvent.cpp
+++ b/Engine/Source/Editor/GameplayAbilitiesEditor/Private/K2Node_GameplayCueEvent.cpp
@@ -4,7 +4,7 @@
 #include "AbilitySystemEditorPrivatePCH.h"
 #include "K2Node_GameplayCueEvent.h"
 #include "CompilerResultsLog.h"
-#include "GameplayTagsModule.h"
+#include "GameplayTags.h"
 #include "GameplayCueInterface.h"
 #include "BlueprintEventNodeSpawner.h"
 #include "BlueprintEditorUtils.h"
@@ -75,11 +75,11 @@ void UK2Node_GameplayCueEvent::GetMenuActions(FBlueprintActionDatabaseRegistrar&
 		EventNode->CustomFunctionName = TagName;
 	};
 	
-	IGameplayTagsModule& GameplayTagsModule = IGameplayTagsModule::Get();
-	FGameplayTag RootTag = GameplayTagsModule.GetGameplayTagsManager().RequestGameplayTag(FName(TEXT("GameplayCue")), false);
+	UGameplayTagsManager& Manager = UGameplayTagsManager::Get();
+	FGameplayTag RootTag = Manager.RequestGameplayTag(FName(TEXT("GameplayCue")), false);
 	if (RootTag.IsValid())
 	{
-		FGameplayTagContainer CueTags = GameplayTagsModule.GetGameplayTagsManager().RequestGameplayTagChildren(RootTag);
+		FGameplayTagContainer CueTags = Manager.RequestGameplayTagChildren(RootTag);
 		// Add a root GameplayCue function as a default
 		CueTags.AddTag(RootTag);
 		for (auto TagIt = CueTags.CreateConstIterator(); TagIt; ++TagIt)
diff --git a/Engine/Source/Editor/GameplayAbilitiesEditor/Private/SGameplayAttributeGraphPin.cpp b/Engine/Source/Editor/GameplayAbilitiesEditor/Private/SGameplayAttributeGraphPin.cpp
--- a/Engine/Source/Editor/GameplayAbilitiesEditor/Private/SGameplayAttributeGraphPin.cpp
+++ b/Engine/Source/Editor/GameplayAbilitiesEditor/Private/SGameplayAttributeGraphPin.cpp
@@ -3,7 +3,6 @@
 #include "AbilitySystemEditorPrivatePCH.h"
 #include "SGameplayAttributeGraphPin.h"
 #include "SGameplayAttributeWidget.h"
-#include "GameplayTagsModule.h"
 #include "GameplayTags.h"
 
 #define LOCTEXT_NAMESPACE "K2Node"
diff --git a/Engine/Source/Editor/GameplayAbilitiesEditor/Private/SGameplayCueEditor.cpp b/Engine/Source/Editor/GameplayAbilitiesEditor/Private/SGameplayCueEditor.cpp
--- a/Engine/Source/Editor/GameplayAbilitiesEditor/Private/SGameplayCueEditor.cpp
+++ b/Engine/Source/Editor/GameplayAbilitiesEditor/Private/SGameplayCueEditor.cpp
@@ -16,10 +16,11 @@
 #include "NotificationManager.h"
 #include "SHyperlink.h"
 #include "SSearchBox.h"
-#include "GameplayTagsModule.h"
+#include "GameplayTags.h"
 #include "AssetEditorManager.h"
 #include "AbilitySystemGlobals.h"
 #include "AssetToolsModule.h"
+#include "GameplayTagsEditorModule.h"
 #include "Editor/LevelEditor/Public/LevelEditor.h"
 
 #define LOCTEXT_NAMESPACE "SGameplayCueEditor"
@@ -147,7 +148,7 @@ public:
 	virtual void OnNewGameplayCueTagCommited(const FText& InText, ETextCommit::Type InCommitType) override
 	{
 		// Only support adding tags via ini file
-		if (UGameplayTagsManager::ShouldImportTagsFromINI() == false)
+		if (UGameplayTagsManager::Get().ShouldImportTagsFromINI() == false)
 		{
 			return;
 		}
@@ -199,7 +200,7 @@ public:
 		SelectedTag = FName(*str);
 		SelectedUniqueID = 0;
 
-		UGameplayTagsManager::AddNewGameplayTagToINI(str);
+		IGameplayTagsEditorModule::Get().AddNewGameplayTagToINI(str);
 
 		UpdateGameplayCueListItems();
 
@@ -435,7 +436,7 @@ BEGIN_SLATE_FUNCTION_BUILD_OPTIMIZATION
 
 					TGuardValue<bool> SupressUpdate(SGameplayCueEditorImpl::bSuppressCueViewUpdate, true);
 					
-					UGameplayTagsManager::AddNewGameplayTagToINI(Item->GameplayCueTagName.ToString());
+					IGameplayTagsEditorModule::Get().AddNewGameplayTagToINI(Item->GameplayCueTagName.ToString());
 				}
 
 				UClass* ParentClass=nullptr;
@@ -692,7 +693,7 @@ BEGIN_SLATE_FUNCTION_BUILD_OPTIMIZATION
 		GameplayCueListItems.Reset();
 		SelectedItem.Reset();
 
-		IGameplayTagsModule& GameplayTagModule = IGameplayTagsModule::Get();
+		UGameplayTagsManager& Manager = UGameplayTagsManager::Get();
 		FString FullSearchString = SearchText.ToString();
 		TArray<FString> SearchStrings;
 		FullSearchString.ParseIntoArrayWS(SearchStrings);
@@ -722,7 +723,15 @@ BEGIN_SLATE_FUNCTION_BUILD_OPTIMIZATION
 		{
 			FString RequestGameplayTagChildrenPerfMessage = FString::Printf(TEXT(" RequestGameplayTagChildren"));
 			SCOPE_LOG_TIME_IN_SECONDS(*RequestGameplayTagChildrenPerfMessage, nullptr)
-			AllGameplayCueTags = IGameplayTagsModule::Get().GetGameplayTagsManager().RequestGameplayTagChildren(UGameplayCueSet::BaseGameplayCueTag(), bShowOnlyLeafTags);
+
+			if (bShowOnlyLeafTags)
+			{
+				AllGameplayCueTags = Manager.RequestGameplayTagChildrenInDictionary(UGameplayCueSet::BaseGameplayCueTag());
+			}
+			else
+			{
+				AllGameplayCueTags = Manager.RequestGameplayTagChildren(UGameplayCueSet::BaseGameplayCueTag());
+			}
 		}
 
 		// Create data structs for widgets
@@ -910,7 +919,7 @@ BEGIN_SLATE_FUNCTION_BUILD_OPTIMIZATION
 
 		EventMap.Empty();
 
-		IGameplayTagsModule& GameplayTagModule = IGameplayTagsModule::Get();
+		UGameplayTagsManager& Manager = UGameplayTagsManager::Get();
 
 		auto del = IGameplayAbilitiesEditorModule::Get().GetGameplayCueInterfaceClassesDelegate();
 		if (del.IsBound())
@@ -952,7 +961,7 @@ BEGIN_SLATE_FUNCTION_BUILD_OPTIMIZATION
 								if (FuncName.Contains("GameplayCue"))
 								{
 									FuncName.ReplaceInline(TEXT("_"), TEXT("."));
-									FGameplayTag FoundTag = GameplayTagModule.RequestGameplayTag(FName(*FuncName), false);
+									FGameplayTag FoundTag = Manager.RequestGameplayTag(FName(*FuncName), false);
 									if (FoundTag.IsValid())
 									{
 										EventMap.AddUnique(FoundTag, *FuncIt);
@@ -1020,7 +1029,7 @@ void SGameplayCueEditorImpl::Construct(const FArguments& InArgs)
 	bHasLoadedAllGameplayCues = false;
 	bFilterIDsDirty = false;
 
-	bool CanAddFromINI = UGameplayTagsManager::ShouldImportTagsFromINI(); // We only support adding new tags to the ini files.
+	bool CanAddFromINI = UGameplayTagsManager::Get().ShouldImportTagsFromINI(); // We only support adding new tags to the ini files.
 	
 	ChildSlot
 	[
diff --git a/Engine/Source/Editor/InputBindingEditor/InputBindingEditor.Build.cs b/Engine/Source/Editor/InputBindingEditor/InputBindingEditor.Build.cs
--- a/Engine/Source/Editor/InputBindingEditor/InputBindingEditor.Build.cs
+++ b/Engine/Source/Editor/InputBindingEditor/InputBindingEditor.Build.cs
@@ -19,6 +19,7 @@ public class InputBindingEditor : ModuleRules
 				"Engine",
 				"UnrealEd",
 				"PropertyEditor",
+				"Settings"
 			}
 		);
 
