Manual merge of CL#4560353 over from the original stream

Original description:
Fixed a crash in FGameplayDebuggerExtension_Spectator::ToggleSpectatorMode resulting from ADebugCameraController.OriginalPlayer being null (for whatever reason)

#rnx

diff --git a/Engine/Source/Developer/GameplayDebugger/Private/GameplayDebuggerExtension_Spectator.cpp b/Engine/Source/Developer/GameplayDebugger/Private/GameplayDebuggerExtension_Spectator.cpp
--- a/Engine/Source/Developer/GameplayDebugger/Private/GameplayDebuggerExtension_Spectator.cpp
+++ b/Engine/Source/Developer/GameplayDebugger/Private/GameplayDebuggerExtension_Spectator.cpp
@@ -80,7 +80,10 @@ void FGameplayDebuggerExtension_Spectator::ToggleSpectatorMode()
 	else
 	{
 		SpectatorControllerOb->PopInputComponent(DebuggerInput);
-		SpectatorControllerOb->OriginalPlayer->SwitchController(OwnerPC);
+		if (SpectatorControllerOb->OriginalPlayer)
+		{
+			SpectatorControllerOb->OriginalPlayer->SwitchController(OwnerPC);
+		}
 		SpectatorControllerOb->OnDeactivate(OwnerPC);
 		OwnerPC->PushInputComponent(DebuggerInput);
 
