When picking pawn in the GameplayDebugger, also selecting it in the editor. That way the BehaviorTreeEditor automatically changes to the right one.

diff --git a/Engine/Source/Developer/GameplayDebugger/Private/GameplayDebuggerCategoryReplicator.cpp b/Engine/Source/Developer/GameplayDebugger/Private/GameplayDebuggerCategoryReplicator.cpp
--- a/Engine/Source/Developer/GameplayDebugger/Private/GameplayDebuggerCategoryReplicator.cpp
+++ b/Engine/Source/Developer/GameplayDebugger/Private/GameplayDebuggerCategoryReplicator.cpp
@@ -1,6 +1,8 @@
 // Copyright 1998-2019 Epic Games, Inc. All Rights Reserved.
 
 #include "GameplayDebuggerCategoryReplicator.h"
+#include "Editor.h"
+#include "Engine/Selection.h"
 #include "GameFramework/PlayerController.h"
 #include "GameplayDebuggerAddonBase.h"
 #include "GameplayDebuggerAddonManager.h"
@@ -503,14 +505,14 @@ void AGameplayDebuggerCategoryReplicator::ServerSetEnabled_Implementation(bool b
 	SetEnabled(bEnable);
 }
 
-bool AGameplayDebuggerCategoryReplicator::ServerSetDebugActor_Validate(AActor* Actor)
+bool AGameplayDebuggerCategoryReplicator::ServerSetDebugActor_Validate(AActor* Actor, bool bSelectInEditor)
 {
 	return true;
 }
 
-void AGameplayDebuggerCategoryReplicator::ServerSetDebugActor_Implementation(AActor* Actor)
+void AGameplayDebuggerCategoryReplicator::ServerSetDebugActor_Implementation(AActor* Actor, bool bSelectInEditor)
 {
-	SetDebugActor(Actor);
+	SetDebugActor(Actor, bSelectInEditor);
 }
 
 bool AGameplayDebuggerCategoryReplicator::ServerSetCategoryEnabled_Validate(int32 CategoryId, bool bEnable)
@@ -610,7 +612,10 @@ void AGameplayDebuggerCategoryReplicator::OnReceivedDataPackPacket(int32 Categor
 void AGameplayDebuggerCategoryReplicator::TickActor(float DeltaTime, enum ELevelTick TickType, FActorTickFunction& ThisTickFunction)
 {
 	Super::TickActor(DeltaTime, TickType, ThisTickFunction);
-	CollectCategoryData();
+	if (OwnerPC)
+	{
+		CollectCategoryData();
+	}
 }
 
 void AGameplayDebuggerCategoryReplicator::PostNetReceive()
@@ -722,7 +727,7 @@ void AGameplayDebuggerCategoryReplicator::SetEnabled(bool bEnable)
 	NotifyExtensionsToolState(bEnableExtensions);
 }
 
-void AGameplayDebuggerCategoryReplicator::SetDebugActor(AActor* Actor)
+void AGameplayDebuggerCategoryReplicator::SetDebugActor(AActor* Actor, bool bSelectInEditor)
 {
 	UE_LOG(LogGameplayDebugReplication, Log, TEXT("SetDebugActor %s"), *GetNameSafe(Actor));
 	if (bHasAuthority)
@@ -732,11 +737,23 @@ void AGameplayDebuggerCategoryReplicator::SetDebugActor(AActor* Actor)
 			DebugActor.Actor = Actor;
 			DebugActor.ActorName = Actor ? Actor->GetFName() : NAME_None;
 			DebugActor.SyncCounter++;
+
+#if WITH_EDITOR
+			if (bSelectInEditor)
+			{
+				USelection* SelectedActors = GEditor ? GEditor->GetSelectedActors() : NULL;
+				if (SelectedActors && Actor)
+				{
+					SelectedActors->DeselectAll();
+					SelectedActors->Select(Actor);
+				}
+			}
+#endif // WITH_EDITOR
 		}
 	}
 	else
 	{
-		ServerSetDebugActor(Actor);
+		ServerSetDebugActor(Actor, bSelectInEditor);
 	}
 }
 
diff --git a/Engine/Source/Developer/GameplayDebugger/Private/GameplayDebuggerLocalController.cpp b/Engine/Source/Developer/GameplayDebugger/Private/GameplayDebuggerLocalController.cpp
--- a/Engine/Source/Developer/GameplayDebugger/Private/GameplayDebuggerLocalController.cpp
+++ b/Engine/Source/Developer/GameplayDebugger/Private/GameplayDebuggerLocalController.cpp
@@ -52,6 +52,7 @@ void UGameplayDebuggerLocalController::Initialize(AGameplayDebuggerCategoryRepli
 	if (GIsEditor)
 	{
 		USelection::SelectionChangedEvent.AddUObject(this, &UGameplayDebuggerLocalController::OnSelectionChanged);
+		USelection::SelectObjectEvent.AddUObject(this, &UGameplayDebuggerLocalController::OnSelectedObject);
 	}
 #endif
 
@@ -102,6 +103,7 @@ void UGameplayDebuggerLocalController::Cleanup()
 {
 #if WITH_EDITOR
 	USelection::SelectionChangedEvent.RemoveAll(this);
+	USelection::SelectObjectEvent.RemoveAll(this);
 
 	if (bSimulateMode)
 	{
@@ -653,7 +655,7 @@ void UGameplayDebuggerLocalController::OnSelectActorTick()
 		if (DebugActorCandidate != BestCandidate)
 		{
 			DebugActorCandidate = BestCandidate;
-			CachedReplicator->SetDebugActor(BestCandidate);
+			CachedReplicator->SetDebugActor(BestCandidate, true);
 		}
 	}
 }
@@ -701,7 +703,21 @@ void UGameplayDebuggerLocalController::OnSelectionChanged(UObject* Object)
 			}
 		}
 
-		CachedReplicator->SetDebugActor(SelectedPawn);
+		if (SelectedPawn)
+		{
+			CachedReplicator->SetDebugActor(SelectedPawn, false);
+			CachedReplicator->CollectCategoryData(/*bForce=*/true);
+		}
+	}
+}
+
+void UGameplayDebuggerLocalController::OnSelectedObject(UObject* Object)
+{
+	AController* SelectedController = Cast<AController>(Object);
+	APawn* SelectedPawn = SelectedController ? SelectedController->GetPawn() : Cast<APawn>(Object);
+	if (CachedReplicator && SelectedPawn && SelectedPawn->IsSelected())
+	{
+		CachedReplicator->SetDebugActor(SelectedPawn, false);
 		CachedReplicator->CollectCategoryData(/*bForce=*/true);
 	}
 }
diff --git a/Engine/Source/Developer/GameplayDebugger/Public/GameplayDebuggerCategoryReplicator.h b/Engine/Source/Developer/GameplayDebugger/Public/GameplayDebuggerCategoryReplicator.h
--- a/Engine/Source/Developer/GameplayDebugger/Public/GameplayDebuggerCategoryReplicator.h
+++ b/Engine/Source/Developer/GameplayDebugger/Public/GameplayDebuggerCategoryReplicator.h
@@ -90,7 +90,7 @@ public:
 	void SetCategoryEnabled(int32 CategoryId, bool bEnable);
 
 	/** [ALL] set actor for debugging */
-	void SetDebugActor(AActor* Actor);
+	void SetDebugActor(AActor* Actor, bool bSelectInEditor);
 
 	/** [ALL] send input event to category */
 	void SendCategoryInputEvent(int32 CategoryId, int32 HandlerId);
@@ -183,7 +183,7 @@ protected:
 	void ServerSetEnabled(bool bEnable);
 
 	UFUNCTION(Server, Reliable, WithValidation, meta = (CallInEditor = "true"))
-	void ServerSetDebugActor(AActor* Actor);
+	void ServerSetDebugActor(AActor* Actor, bool bSelectInEditor);
 
 	UFUNCTION(Server, Reliable, WithValidation, meta = (CallInEditor = "true"))
 	void ServerSetCategoryEnabled(int32 CategoryId, bool bEnable);
diff --git a/Engine/Source/Developer/GameplayDebugger/Public/GameplayDebuggerLocalController.h b/Engine/Source/Developer/GameplayDebugger/Public/GameplayDebuggerLocalController.h
--- a/Engine/Source/Developer/GameplayDebugger/Public/GameplayDebuggerLocalController.h
+++ b/Engine/Source/Developer/GameplayDebugger/Public/GameplayDebuggerLocalController.h
@@ -112,6 +112,7 @@ protected:
 
 	/** event for simulate in editor mode */
 	void OnSelectionChanged(UObject* Object);
+	void OnSelectedObject(UObject* Object);
 
 	FString GetKeyDescriptionShort(const FKey& KeyBind) const;
 	FString GetKeyDescriptionLong(const FKey& KeyBind) const;
