Fix build break

diff --git a/Engine/Source/Developer/GameplayDebugger/Private/GameplayDebuggerCategoryReplicator.cpp b/Engine/Source/Developer/GameplayDebugger/Private/GameplayDebuggerCategoryReplicator.cpp
--- a/Engine/Source/Developer/GameplayDebugger/Private/GameplayDebuggerCategoryReplicator.cpp
+++ b/Engine/Source/Developer/GameplayDebugger/Private/GameplayDebuggerCategoryReplicator.cpp
@@ -1,8 +1,6 @@
 // Copyright 1998-2019 Epic Games, Inc. All Rights Reserved.
 
 #include "GameplayDebuggerCategoryReplicator.h"
-#include "Editor.h"
-#include "Engine/Selection.h"
 #include "GameFramework/PlayerController.h"
 #include "GameplayDebuggerAddonBase.h"
 #include "GameplayDebuggerAddonManager.h"
@@ -11,6 +9,11 @@
 #include "GameplayDebuggerExtension.h"
 #include "Net/UnrealNetwork.h"
 
+#if WITH_EDITOR
+#include "Editor.h"
+#include "Engine/Selection.h"
+#endif // WITH_EDITOR
+
 //////////////////////////////////////////////////////////////////////////
 // FGameplayDebuggerCategoryReplicatorData
 
