Fix build break

diff --git a/Engine/Source/Developer/GameplayDebugger/Public/GameplayDebuggerCategoryReplicator.h b/Engine/Source/Developer/GameplayDebugger/Public/GameplayDebuggerCategoryReplicator.h
--- a/Engine/Source/Developer/GameplayDebugger/Public/GameplayDebuggerCategoryReplicator.h
+++ b/Engine/Source/Developer/GameplayDebugger/Public/GameplayDebuggerCategoryReplicator.h
@@ -90,7 +90,7 @@ public:
 	void SetCategoryEnabled(int32 CategoryId, bool bEnable);
 
 	/** [ALL] set actor for debugging */
-	void SetDebugActor(AActor* Actor, bool bSelectInEditor);
+	void SetDebugActor(AActor* Actor, bool bSelectInEditor = false);
 
 	/** [ALL] send input event to category */
 	void SendCategoryInputEvent(int32 CategoryId, int32 HandlerId);
