fixed gameplay debugger crash on replay recording

diff --git a/Engine/Source/Developer/GameplayDebugger/Private/GameplayDebuggerCategoryReplicator.cpp b/Engine/Source/Developer/GameplayDebugger/Private/GameplayDebuggerCategoryReplicator.cpp
--- a/Engine/Source/Developer/GameplayDebugger/Private/GameplayDebuggerCategoryReplicator.cpp
+++ b/Engine/Source/Developer/GameplayDebugger/Private/GameplayDebuggerCategoryReplicator.cpp
@@ -102,6 +102,14 @@ bool FGameplayDebuggerNetPack::NetDeltaSerialize(FNetDeltaSerializeInfo& DeltaPa
 
 	if (DeltaParms.Writer)
 	{
+		// watch for DemoNetDriver doing snapshot on client, no need to store debug data here
+		// especially if it tries to write incomplete packets
+		const bool bIsOwnerClient = !Owner->bHasAuthority;
+		if (bIsOwnerClient)
+		{
+			return false;
+		}
+
 		FBitWriter& Writer = *DeltaParms.Writer;
 		int32 NumChangedCategories = 0;
 
