Added a piece of information to GameplayDebugger that allows syncing it to specific VisLogger recording file (if recoding-to-file is in progress).

diff --git a/Engine/Source/Developer/GameplayDebugger/Private/GameplayDebuggerCategoryReplicator.cpp b/Engine/Source/Developer/GameplayDebugger/Private/GameplayDebuggerCategoryReplicator.cpp
--- a/Engine/Source/Developer/GameplayDebugger/Private/GameplayDebuggerCategoryReplicator.cpp
+++ b/Engine/Source/Developer/GameplayDebugger/Private/GameplayDebuggerCategoryReplicator.cpp
@@ -8,6 +8,7 @@
 #include "GameplayDebuggerRenderingComponent.h"
 #include "GameplayDebuggerExtension.h"
 #include "Net/UnrealNetwork.h"
+#include "VisualLogger/VisualLogger.h"
 
 #if WITH_EDITOR
 #include "Editor.h"
@@ -502,6 +503,7 @@ void AGameplayDebuggerCategoryReplicator::GetLifetimeReplicatedProps(TArray<FLif
 
 	DOREPLIFETIME(AGameplayDebuggerCategoryReplicator, OwnerPC);
 	DOREPLIFETIME(AGameplayDebuggerCategoryReplicator, DebugActor);
+	DOREPLIFETIME(AGameplayDebuggerCategoryReplicator, VisLogSync);
 	DOREPLIFETIME(AGameplayDebuggerCategoryReplicator, bIsEnabled);
 	DOREPLIFETIME(AGameplayDebuggerCategoryReplicator, ReplicatedData);
 }
@@ -644,6 +646,18 @@ void AGameplayDebuggerCategoryReplicator::PostNetReceive()
 void AGameplayDebuggerCategoryReplicator::CollectCategoryData(bool bForce)
 {
 	const float GameTime = GetWorld()->GetTimeSeconds();
+
+	const FVisualLogger& VLogger = FVisualLogger::Get();
+	VisLogSync.DeviceIDs.Reset();
+	if (VLogger.IsRecordingToFile())
+	{
+		const TArray<FVisualLogDevice*>& LogDevices = VLogger.GetDevices();
+		for (const FVisualLogDevice* Device : LogDevices)
+		{
+			VisLogSync.DeviceIDs += FString::Printf(TEXT("%u, "), Device ? Device->GetShortID() : 0);
+		}
+	}
+
 	for (int32 Idx = 0; Idx < Categories.Num(); Idx++)
 	{
 		FGameplayDebuggerCategory& CategoryOb = Categories[Idx].Get();
@@ -713,12 +727,18 @@ void AGameplayDebuggerCategoryReplicator::SetReplicatorOwner(APlayerController*
 			APlayerController* OldOwner = OwnerPC;
 			OwnerPC = InOwnerPC;
 			NotifyDebuggerOwnerChange.Broadcast(this, OldOwner);
+
+			UE_VLOG_UELOG(this, LogGameplayDebugReplication, Log, TEXT("Set OWNER PC %s / %s / %s")
+				, *GetNameSafe(InOwnerPC), InOwnerPC ? *GetNameSafe(InOwnerPC->GetPawn()) : TEXT_EMPTY
+				, InOwnerPC ? *InOwnerPC->GetHumanReadableName() : TEXT_EMPTY);
 		}
 	}
 }
 
 void AGameplayDebuggerCategoryReplicator::SetEnabled(bool bEnable)
 {
+	UE_VLOG_UELOG(this, LogGameplayDebugReplication, Log, TEXT("SetEnabled %s"), TEXT_CONDITION(bEnable));
+
 	if (bHasAuthority)
 	{
 		bIsEnabled = bEnable;
@@ -740,7 +760,7 @@ void AGameplayDebuggerCategoryReplicator::SetEnabled(bool bEnable)
 
 void AGameplayDebuggerCategoryReplicator::SetDebugActor(AActor* Actor, bool bSelectInEditor)
 {
-	UE_LOG(LogGameplayDebugReplication, Log, TEXT("SetDebugActor %s"), *GetNameSafe(Actor));
+	UE_VLOG_UELOG(this, LogGameplayDebugReplication, Log, TEXT("SetDebugActor %s"), *GetNameSafe(Actor));	
 	if (bHasAuthority)
 	{
 		if (DebugActor.Actor != Actor)
@@ -774,7 +794,7 @@ void AGameplayDebuggerCategoryReplicator::SetCategoryEnabled(int32 CategoryId, b
 	{
 		if (Categories.IsValidIndex(CategoryId))
 		{
-			UE_LOG(LogGameplayDebugReplication, Log, TEXT("SetCategoryEnabled[%d]:%d (%s)"), CategoryId, bEnable ? 1 : 0, *Categories[CategoryId]->GetCategoryName().ToString());
+			UE_VLOG_UELOG(this, LogGameplayDebugReplication, Log, TEXT("SetCategoryEnabled[%d]:%d (%s)"), CategoryId, bEnable ? 1 : 0, *Categories[CategoryId]->GetCategoryName().ToString());
 			Categories[CategoryId]->bIsEnabled = bEnable;
 		}
 	}
diff --git a/Engine/Source/Developer/GameplayDebugger/Private/GameplayDebuggerLocalController.cpp b/Engine/Source/Developer/GameplayDebugger/Private/GameplayDebuggerLocalController.cpp
--- a/Engine/Source/Developer/GameplayDebugger/Private/GameplayDebuggerLocalController.cpp
+++ b/Engine/Source/Developer/GameplayDebugger/Private/GameplayDebuggerLocalController.cpp
@@ -211,6 +211,13 @@ void UGameplayDebuggerLocalController::DrawHeader(FGameplayDebuggerCanvasContext
 	CanvasContext.MeasureString(DebugActorDesc, DebugActorSizeX, DebugActorSizeY);
 	CanvasContext.PrintAt(CanvasContext.Canvas->SizeX - PaddingRight - DebugActorSizeX, UsePaddingTop, DebugActorDesc);
 
+	const FString VLogDesc = FString::Printf(TEXT("VLog: {cyan}%s"), CachedReplicator->GetVisLogSyncData().DeviceIDs.Len() > 0
+			? *CachedReplicator->GetVisLogSyncData().DeviceIDs
+			: TEXT("not recording to file"));
+	float VLogSizeX = 0.0f, VLogSizeY = 0.0f;
+	CanvasContext.MeasureString(VLogDesc, VLogSizeX, VLogSizeY);
+	CanvasContext.PrintAt(CanvasContext.Canvas->SizeX - PaddingRight - VLogSizeX, UsePaddingTop + LineHeight, VLogDesc);
+
 	const FString TimestampDesc = FString::Printf(TEXT("Time: %.2fs"), CachedReplicator->GetWorld()->GetTimeSeconds());
 	float TimestampSizeX = 0.0f, TimestampSizeY = 0.0f;
 	CanvasContext.MeasureString(TimestampDesc, TimestampSizeX, TimestampSizeY);
diff --git a/Engine/Source/Developer/GameplayDebugger/Public/GameplayDebuggerCategoryReplicator.h b/Engine/Source/Developer/GameplayDebugger/Public/GameplayDebuggerCategoryReplicator.h
--- a/Engine/Source/Developer/GameplayDebugger/Public/GameplayDebuggerCategoryReplicator.h
+++ b/Engine/Source/Developer/GameplayDebugger/Public/GameplayDebuggerCategoryReplicator.h
@@ -61,6 +61,15 @@ struct FGameplayDebuggerDebugActor
 	int32 SyncCounter;
 };
 
+USTRUCT()
+struct FGameplayDebuggerVisLogSync
+{
+	GENERATED_USTRUCT_BODY()
+
+	UPROPERTY()
+	FString DeviceIDs;
+};
+
 DECLARE_MULTICAST_DELEGATE_TwoParams(FNotifyGameplayDebuggerOwnerChange, AGameplayDebuggerCategoryReplicator*, APlayerController* /** Old Owner */);
 
 UCLASS(NotBlueprintable, NotBlueprintType, notplaceable, noteditinlinenew, hidedropdown, Transient)
@@ -110,6 +119,8 @@ public:
 	/** get sync counter, increased with every change of DebugActor */
 	int16 GetDebugActorCounter() const { return DebugActor.SyncCounter; }
 
+	const FGameplayDebuggerVisLogSync& GetVisLogSyncData() const { return VisLogSync; }
+
 	/** get player controller owning this replicator */
 	APlayerController* GetReplicationOwner() const { return OwnerPC; }
 
@@ -153,6 +164,9 @@ protected:
 	UPROPERTY(Replicated)
 	FGameplayDebuggerDebugActor	DebugActor;
 
+	UPROPERTY(Replicated)
+	FGameplayDebuggerVisLogSync	VisLogSync;
+
 	/** rendering component needs to attached to some actor, and this is as good as any */
 	UPROPERTY()
 	UGameplayDebuggerRenderingComponent* RenderingComp;
diff --git a/Engine/Source/Runtime/Engine/Private/VisualLogger/VisualLoggerBinaryFileDevice.cpp b/Engine/Source/Runtime/Engine/Private/VisualLogger/VisualLoggerBinaryFileDevice.cpp
--- a/Engine/Source/Runtime/Engine/Private/VisualLogger/VisualLoggerBinaryFileDevice.cpp
+++ b/Engine/Source/Runtime/Engine/Private/VisualLogger/VisualLoggerBinaryFileDevice.cpp
@@ -61,14 +61,15 @@ void FVisualLoggerBinaryFileDevice::StopRecordingToFile(float TimeStamp)
 	FileArchive = nullptr;
 
 	const FString TempFullFilename = FPaths::Combine(*FPaths::ProjectLogDir(), *TempFileName);
-	const FString NewFileName = FVisualLoggerHelpers::GenerateFilename(TempFileName, FileName, StartRecordingTime, LastLogTimeStamp);
+	const FString NewFileName = FString::Printf(TEXT("%u_%s"), GetShortID(), *FVisualLoggerHelpers::GenerateFilename(TempFileName, FileName, StartRecordingTime, LastLogTimeStamp));
+	const FString NewFullFileName = FPaths::Combine(*FPaths::ProjectLogDir(), *NewFileName);
 
 	if (TotalSize > 0)
 	{
 		// rename file when we serialized some data
-		IFileManager::Get().Move(*NewFileName, *TempFullFilename, true);
+		IFileManager::Get().Move(*NewFullFileName, *TempFullFilename, true);
 
-		UE_LOG(LogVisual, Display, TEXT("Vislog file saved: %s"), *NewFileName);
+		UE_LOG(LogVisual, Display, TEXT("Vislog file saved: %s"), *NewFullFileName);
 	}
 	else
 	{
diff --git a/Engine/Source/Runtime/Engine/Private/VisualLogger/VisualLoggerTypes.cpp b/Engine/Source/Runtime/Engine/Private/VisualLogger/VisualLoggerTypes.cpp
--- a/Engine/Source/Runtime/Engine/Private/VisualLogger/VisualLoggerTypes.cpp
+++ b/Engine/Source/Runtime/Engine/Private/VisualLogger/VisualLoggerTypes.cpp
@@ -20,6 +20,17 @@ UVisualLoggerDebugSnapshotInterface::UVisualLoggerDebugSnapshotInterface(const F
 #define DEPRECATED_VISUAL_LOGGER_MAGIC_NUMBER 0xFAFAAFAF
 #define VISUAL_LOGGER_MAGIC_NUMBER 0xAFAFFAFA
 
+//----------------------------------------------------------------------//
+// FVisualLogDevice 
+//----------------------------------------------------------------------//
+FVisualLogDevice::FVisualLogDevice()
+	: GUID(FGuid::NewGuid())
+{
+}
+
+//----------------------------------------------------------------------//
+// FVisualLogShapeElement 
+//----------------------------------------------------------------------//
 FVisualLogShapeElement::FVisualLogShapeElement(EVisualLoggerShapeElement InType)
 	: Category(NAME_UnnamedCategory)
 	, Verbosity(ELogVerbosity::All)
@@ -493,8 +504,7 @@ FString FVisualLoggerHelpers::GenerateTemporaryFilename(const FString& FileExt)
 
 FString FVisualLoggerHelpers::GenerateFilename(const FString& TempFileName, const FString& Prefix, float StartRecordingTime, float EndTimeStamp)
 {
-	const FString TempFullFilename = FString::Printf(TEXT("%s%s"), *FPaths::ProjectLogDir(), *TempFileName);
-	const FString FullFilename = FString::Printf(TEXT("%s%s_%s"), *FPaths::ProjectLogDir(), *Prefix, *TempFileName);
+	const FString FullFilename = FString::Printf(TEXT("%s_%s"), *Prefix, *TempFileName);
 	const FString TimeFrameString = FString::Printf(TEXT("%d-%d_"), FMath::TruncToInt(StartRecordingTime), FMath::TruncToInt(EndTimeStamp));
 	return FullFilename.Replace(TEXT("VTEMP_"), *TimeFrameString, ESearchCase::CaseSensitive);
 }
diff --git a/Engine/Source/Runtime/Engine/Public/VisualLogger/VisualLoggerTypes.h b/Engine/Source/Runtime/Engine/Public/VisualLogger/VisualLoggerTypes.h
--- a/Engine/Source/Runtime/Engine/Public/VisualLogger/VisualLoggerTypes.h
+++ b/Engine/Source/Runtime/Engine/Public/VisualLogger/VisualLoggerTypes.h
@@ -266,6 +266,7 @@ public:
 		FVisualLogEntry Entry;
 	};
 
+	FVisualLogDevice();
 	virtual ~FVisualLogDevice() { }
 	virtual void Serialize(const UObject* LogOwner, FName OwnerName, FName InOwnerClassName, const FVisualLogEntry& LogEntry) = 0;
 	virtual void Cleanup(bool bReleaseMemory = false) { /* Empty */ }
@@ -275,6 +276,9 @@ public:
 	virtual void SetFileName(const FString& InFileName) { /* Empty */ }
 	virtual void GetRecordedLogs(TArray<FVisualLogDevice::FVisualLogEntryItem>& OutLogs)  const { /* Empty */ }
 	virtual bool HasFlags(int32 InFlags) const { return false; }
+	uint32 GetShortID() const { return GUID[0]; }
+
+	const FGuid GUID;
 };
 
 struct ENGINE_API FVisualLoggerCategoryVerbosityPair
