Fixed CIS build for GameplayDebuggerCategoryReplicator when ENABLE_VISUAL_LOG is not defined

#rnx

diff --git a/Engine/Source/Developer/GameplayDebugger/Private/GameplayDebuggerCategoryReplicator.cpp b/Engine/Source/Developer/GameplayDebugger/Private/GameplayDebuggerCategoryReplicator.cpp
--- a/Engine/Source/Developer/GameplayDebugger/Private/GameplayDebuggerCategoryReplicator.cpp
+++ b/Engine/Source/Developer/GameplayDebugger/Private/GameplayDebuggerCategoryReplicator.cpp
@@ -647,6 +647,7 @@ void AGameplayDebuggerCategoryReplicator::CollectCategoryData(bool bForce)
 {
 	const float GameTime = GetWorld()->GetTimeSeconds();
 
+#if ENABLE_VISUAL_LOG
 	const FVisualLogger& VLogger = FVisualLogger::Get();
 	VisLogSync.DeviceIDs.Reset();
 	if (VLogger.IsRecordingToFile())
@@ -657,6 +658,7 @@ void AGameplayDebuggerCategoryReplicator::CollectCategoryData(bool bForce)
 			VisLogSync.DeviceIDs += FString::Printf(TEXT("%u, "), Device ? Device->GetShortID() : 0);
 		}
 	}
+#endif // ENABLE_VISUAL_LOG
 
 	for (int32 Idx = 0; Idx < Categories.Num(); Idx++)
 	{
