CIS fix

#rnx

diff --git a/Engine/Source/Developer/GameplayDebugger/Private/GameplayDebuggerCategoryReplicator.cpp b/Engine/Source/Developer/GameplayDebugger/Private/GameplayDebuggerCategoryReplicator.cpp
--- a/Engine/Source/Developer/GameplayDebugger/Private/GameplayDebuggerCategoryReplicator.cpp
+++ b/Engine/Source/Developer/GameplayDebugger/Private/GameplayDebuggerCategoryReplicator.cpp
@@ -655,7 +655,7 @@ void AGameplayDebuggerCategoryReplicator::CollectCategoryData(bool bForce)
 		const TArray<FVisualLogDevice*>& LogDevices = VLogger.GetDevices();
 		for (const FVisualLogDevice* Device : LogDevices)
 		{
-			VisLogSync.DeviceIDs += FString::Printf(TEXT("%u, "), Device ? Device->GetShortID() : 0);
+			VisLogSync.DeviceIDs += FString::Printf(TEXT("%u, "), Device ? Device->GetShortSessionID() : 0);
 		}
 	}
 #endif // ENABLE_VISUAL_LOG
