Gameplay Debugger:
- Fix HUD not hiding
- Changed call to hide/show HUD to use the ShowHUD function instead of setting the bool directly so we can override the function

#rnx

diff --git a/Engine/Source/Developer/GameplayDebugger/Private/GameplayDebuggerExtension_HUD.cpp b/Engine/Source/Developer/GameplayDebugger/Private/GameplayDebuggerExtension_HUD.cpp
--- a/Engine/Source/Developer/GameplayDebugger/Private/GameplayDebuggerExtension_HUD.cpp
+++ b/Engine/Source/Developer/GameplayDebugger/Private/GameplayDebuggerExtension_HUD.cpp
@@ -82,7 +82,11 @@ void FGameplayDebuggerExtension_HUD::SetGameHUDEnabled(bool bEnable)
 	AHUD* GameHUD = OwnerPC ? OwnerPC->GetHUD() : nullptr;
 	if (GameHUD)
 	{
-		GameHUD->bShowHUD = bEnable;
+		// ShowHUD is actually a toggle and is a UFUNCTION(exec).
+		// We make sure the flag is !bEnable before calling the function so we don't end up desynced.
+		// Also, we don't simply set the flag to bEnable because we want to execute functions that overrides ShowHUD.
+		GameHUD->bShowHUD = !bEnable;
+		GameHUD->ShowHUD();
 	}
 
 	bIsGameHUDEnabled = bEnable;
