GameplayDebugger console commands:
+ EnableGDT (for legacy support)
+ gdt.Toggle
+ gdt.SelectPreviousRow
+ gdt.SelectNextRow
+ gdt.ToggleCategory

diff --git a/Engine/Source/Developer/GameplayDebugger/Private/GameplayDebuggerLocalController.cpp b/Engine/Source/Developer/GameplayDebugger/Private/GameplayDebuggerLocalController.cpp
--- a/Engine/Source/Developer/GameplayDebugger/Private/GameplayDebuggerLocalController.cpp
+++ b/Engine/Source/Developer/GameplayDebugger/Private/GameplayDebuggerLocalController.cpp
@@ -25,6 +25,7 @@
 #include "Engine/LocalPlayer.h"
 #include "GameFramework/PlayerInput.h"
 #include "EngineUtils.h"
+#include "HAL/IConsoleManager.h"
 
 #if WITH_EDITOR
 #include "Editor/GameplayDebuggerEdMode.h"
@@ -160,7 +161,7 @@ extern RENDERCORE_API FTexture* GWhiteTexture;
 
 void UGameplayDebuggerLocalController::DrawHeader(FGameplayDebuggerCanvasContext& CanvasContext)
 {
-	const int32 NumRows = (NumCategorySlots + 9) / 10;
+	const int32 NumRows = (NumCategorySlots + (NumCategoriesPerRow-1)) / NumCategoriesPerRow;
 	const float LineHeight = CanvasContext.GetLineHeight();
 	const int32 NumExtensions = bSimulateMode ? 0 : CachedReplicator->GetNumExtensions();
 	const int32 NumExtensionRows = (NumExtensions > 0) ? 1 : 0;
@@ -257,9 +258,9 @@ void UGameplayDebuggerLocalController::DrawHeader(FGameplayDebuggerCanvasContext
 	for (int32 RowIdx = 0; RowIdx < NumRows; RowIdx++)
 	{
 		FString CategoryRowDesc;
-		for (int32 Idx = 0; Idx < 10; Idx++)
+		for (int32 Idx = 0; Idx < NumCategoriesPerRow; Idx++)
 		{
-			const int32 CategorySlotIdx = (RowIdx * 10) + Idx;
+			const int32 CategorySlotIdx = (RowIdx * NumCategoriesPerRow) + Idx;
 			if (SlotCategoryIds.IsValidIndex(CategorySlotIdx) && 
 				SlotNames.IsValidIndex(CategorySlotIdx) &&
 				SlotCategoryIds[CategorySlotIdx].Num())
@@ -466,11 +467,16 @@ void UGameplayDebuggerLocalController::OnActivationPressed()
 }
 
 void UGameplayDebuggerLocalController::OnActivationReleased()
+{
+	ToggleActivation();
+}
+
+void UGameplayDebuggerLocalController::ToggleActivation()
 {
 	if (CachedReplicator)
 	{
 		UWorld* World = CachedReplicator->GetWorld();
-		if (StartSelectingActorHandle.IsValid())
+		if (!bIsSelectingActor || StartSelectingActorHandle.IsValid())
 		{
 			bIsLocallyEnabled = !CachedReplicator->IsEnabled();
 			CachedReplicator->SetEnabled(bIsLocallyEnabled);
@@ -500,63 +506,63 @@ void UGameplayDebuggerLocalController::OnActivationReleased()
 
 void UGameplayDebuggerLocalController::OnCategory0Pressed()
 {
-	ToggleSlotState((ActiveRowIdx * 10) + 0);
+	ToggleSlotState((ActiveRowIdx * NumCategoriesPerRow) + 0);
 }
 
 void UGameplayDebuggerLocalController::OnCategory1Pressed()
 {
-	ToggleSlotState((ActiveRowIdx * 10) + 1);
+	ToggleSlotState((ActiveRowIdx * NumCategoriesPerRow) + 1);
 }
 
 void UGameplayDebuggerLocalController::OnCategory2Pressed()
 {
-	ToggleSlotState((ActiveRowIdx * 10) + 2);
+	ToggleSlotState((ActiveRowIdx * NumCategoriesPerRow) + 2);
 }
 
 void UGameplayDebuggerLocalController::OnCategory3Pressed()
 {
-	ToggleSlotState((ActiveRowIdx * 10) + 3);
+	ToggleSlotState((ActiveRowIdx * NumCategoriesPerRow) + 3);
 }
 
 void UGameplayDebuggerLocalController::OnCategory4Pressed()
 {
-	ToggleSlotState((ActiveRowIdx * 10) + 4);
+	ToggleSlotState((ActiveRowIdx * NumCategoriesPerRow) + 4);
 }
 
 void UGameplayDebuggerLocalController::OnCategory5Pressed()
 {
-	ToggleSlotState((ActiveRowIdx * 10) + 5);
+	ToggleSlotState((ActiveRowIdx * NumCategoriesPerRow) + 5);
 }
 
 void UGameplayDebuggerLocalController::OnCategory6Pressed()
 {
-	ToggleSlotState((ActiveRowIdx * 10) + 6);
+	ToggleSlotState((ActiveRowIdx * NumCategoriesPerRow) + 6);
 }
 
 void UGameplayDebuggerLocalController::OnCategory7Pressed()
 {
-	ToggleSlotState((ActiveRowIdx * 10) + 7);
+	ToggleSlotState((ActiveRowIdx * NumCategoriesPerRow) + 7);
 }
 
 void UGameplayDebuggerLocalController::OnCategory8Pressed()
 {
-	ToggleSlotState((ActiveRowIdx * 10) + 8);
+	ToggleSlotState((ActiveRowIdx * NumCategoriesPerRow) + 8);
 }
 
 void UGameplayDebuggerLocalController::OnCategory9Pressed()
 {
-	ToggleSlotState((ActiveRowIdx * 10) + 9);
+	ToggleSlotState((ActiveRowIdx * NumCategoriesPerRow) + 9);
 }
 
 void UGameplayDebuggerLocalController::OnCategoryRowUpPressed()
 {
-	const int32 NumRows = (NumCategorySlots + 9) / 10;
+	const int32 NumRows = (NumCategorySlots + (NumCategoriesPerRow-1)) / NumCategoriesPerRow;
 	ActiveRowIdx = (NumRows > 1) ? ((ActiveRowIdx + NumRows - 1) % NumRows) : 0;
 }
 
 void UGameplayDebuggerLocalController::OnCategoryRowDownPressed()
 {
-	const int32 NumRows = (NumCategorySlots + 9) / 10;
+	const int32 NumRows = (NumCategorySlots + (NumCategoriesPerRow-1)) / NumCategoriesPerRow;
 	ActiveRowIdx = (NumRows > 1) ? ((ActiveRowIdx + 1) % NumRows) : 0;
 }
 
@@ -793,3 +799,123 @@ void UGameplayDebuggerLocalController::RebuildDataPackMap()
 		}
 	}
 }
+
+/**  Helper structure to declare/define console commands in the source file and to access UGameplayDebuggerLocalController protected members */
+struct FGameplayDebuggerConsoleCommands
+{
+private:
+	static UGameplayDebuggerLocalController* GetController(UWorld* InWorld)
+	{
+		UGameplayDebuggerLocalController* Controller = nullptr;
+
+		APlayerController* LocalPC = GEngine->GetFirstLocalPlayerController(InWorld);
+		if (LocalPC)
+		{
+			Controller = AGameplayDebuggerPlayerManager::GetCurrent(InWorld).GetLocalController(*LocalPC);
+		}
+
+		UE_CLOG(Controller == nullptr, LogConsoleResponse, Error, TEXT("GameplayDebugger not available"));
+		return Controller;
+	}
+
+	static void ToggleGameplayDebugger(UWorld* InWorld)
+	{
+		if (UGameplayDebuggerLocalController* Controller = GetController(InWorld))
+		{
+			Controller->ToggleActivation();
+		}
+	}
+
+	static void SelectPreviousRow(UWorld* InWorld)
+	{
+		if (UGameplayDebuggerLocalController* Controller = GetController(InWorld))
+		{
+			Controller->OnCategoryRowUpPressed();
+		}
+	}
+
+	static void SelectNextRow(UWorld* InWorld)
+	{
+		if (UGameplayDebuggerLocalController* Controller = GetController(InWorld))
+		{
+			Controller->OnCategoryRowDownPressed();
+		}
+	}
+
+	static void ToggleCategory(const TArray<FString>& Args, UWorld* InWorld)
+	{
+		UGameplayDebuggerLocalController* Controller = GetController(InWorld);
+		if (Controller == nullptr)
+		{
+			return;
+		}
+
+		if (Args.Num() != 1)
+		{
+			UE_LOG(LogConsoleResponse, Error, TEXT("Missing category index parameter. Usage: gdt.ToggleCategory <CategoryIdx>"));
+			return;
+		}
+
+		if (!Args[0].IsNumeric())
+		{
+			UE_LOG(LogConsoleResponse, Error, TEXT("Must provide numerical value as index. Usage: gdt.ToggleCategory <CategoryIdx>"));
+			return;
+		}
+		
+		const int32 SlotIdx = TCString<TCHAR>::Atoi(*Args[0]);
+		const int32 NumSlots = Controller->SlotCategoryIds.Num();
+		const int32 NumSlotsPerRow = UGameplayDebuggerLocalController::NumCategoriesPerRow;
+		const int32 NumRows = (NumSlots + (NumSlotsPerRow-1)) / NumSlotsPerRow;
+
+		const bool bIsLastRowActive = (Controller->ActiveRowIdx == NumRows-1);
+		const int32 NumSlotsOnActiveRow = bIsLastRowActive ? NumSlots - (NumSlotsPerRow * (NumRows-1)) : NumSlotsPerRow;
+		const int32 MaxSlotIdx = FMath::Max(0, FMath::Min(NumSlots, NumSlotsOnActiveRow)-1);
+		
+		if (!(Controller->SlotCategoryIds.IsValidIndex(SlotIdx) && SlotIdx <= MaxSlotIdx))
+		{
+			UE_LOG(LogConsoleResponse, Error, TEXT("Requires a category index in the active row [0..%d]. Usage: gdt.ToggleCategory CategoryIndex"), MaxSlotIdx);
+			return;
+		}
+		
+		Controller->ToggleSlotState((Controller->ActiveRowIdx * NumSlotsPerRow)+SlotIdx);
+	}
+
+	/** For legacy command: EnableGDT */
+	static FAutoConsoleCommandWithWorld EnableDebuggerCmd;
+
+	/** Various gameplay debugger commands: gdt.<command> */
+	static FAutoConsoleCommandWithWorld ToggleDebuggerCmd;
+	static FAutoConsoleCommandWithWorld SelectPreviousRowCmd;
+	static FAutoConsoleCommandWithWorld SelectNextRowCmd;
+	static FAutoConsoleCommandWithWorldAndArgs ToggleCategoryCmd;
+};
+
+FAutoConsoleCommandWithWorld FGameplayDebuggerConsoleCommands::EnableDebuggerCmd(
+	TEXT("EnableGDT"),
+	TEXT("Toggles Gameplay Debugger Tool"),
+	FConsoleCommandWithWorldDelegate::CreateStatic(&FGameplayDebuggerConsoleCommands::ToggleGameplayDebugger)
+);
+
+FAutoConsoleCommandWithWorld FGameplayDebuggerConsoleCommands::ToggleDebuggerCmd(
+	TEXT("gdt.Toggle"),
+	TEXT("Toggles Gameplay Debugger Tool"),
+	FConsoleCommandWithWorldDelegate::CreateStatic(&FGameplayDebuggerConsoleCommands::ToggleGameplayDebugger)
+);
+
+FAutoConsoleCommandWithWorld FGameplayDebuggerConsoleCommands::SelectPreviousRowCmd(
+	TEXT("gdt.SelectPreviousRow"),
+	TEXT("Selects previous row"),
+	FConsoleCommandWithWorldDelegate::CreateStatic(FGameplayDebuggerConsoleCommands::SelectPreviousRow)
+);
+
+FAutoConsoleCommandWithWorld FGameplayDebuggerConsoleCommands::SelectNextRowCmd(
+	TEXT("gdt.SelectNextRow"),
+	TEXT("Selects next row"),
+	FConsoleCommandWithWorldDelegate::CreateStatic(FGameplayDebuggerConsoleCommands::SelectNextRow)
+);
+
+FAutoConsoleCommandWithWorldAndArgs FGameplayDebuggerConsoleCommands::ToggleCategoryCmd(
+	TEXT("gdt.ToggleCategory"),
+	TEXT("Toggles specific category index"),
+	FConsoleCommandWithWorldAndArgsDelegate::CreateStatic(&FGameplayDebuggerConsoleCommands::ToggleCategory)
+);
diff --git a/Engine/Source/Developer/GameplayDebugger/Public/GameplayDebuggerLocalController.h b/Engine/Source/Developer/GameplayDebugger/Public/GameplayDebuggerLocalController.h
--- a/Engine/Source/Developer/GameplayDebugger/Public/GameplayDebuggerLocalController.h
+++ b/Engine/Source/Developer/GameplayDebugger/Public/GameplayDebuggerLocalController.h
@@ -39,6 +39,7 @@ class UGameplayDebuggerLocalController : public UObject
 	bool IsKeyBound(const FName KeyName) const;
 
 protected:
+	friend struct FGameplayDebuggerConsoleCommands;
 
 	UPROPERTY()
 	AGameplayDebuggerCategoryReplicator* CachedReplicator;
@@ -69,6 +70,7 @@ protected:
 	int32 ActiveRowIdx;
 	int32 NumCategorySlots;
 	int32 NumCategories;
+	static constexpr int32 NumCategoriesPerRow = 10;
 
 	float PaddingLeft;
 	float PaddingRight;
@@ -104,6 +106,9 @@ protected:
 	/** toggle state of categories in given slot */
 	void ToggleSlotState(int32 SlotIdx);
 
+	/** toggle debugger on/off */
+	void ToggleActivation();
+
 	/** draw header row */
 	void DrawHeader(FGameplayDebuggerCanvasContext& CanvasContext);
 
@@ -122,4 +127,4 @@ protected:
 
 	/** build DataPackMap for replication details */
 	void RebuildDataPackMap();
-};
+};
\ No newline at end of file
