Changed actor filtering in UGameplayDebuggerLocalController so that it allows regular (non-pawn) AActors selection in the "editor mode" (used while Simulating-in-Editor).

diff --git a/Engine/Source/Developer/GameplayDebugger/Private/GameplayDebuggerLocalController.cpp b/Engine/Source/Developer/GameplayDebugger/Private/GameplayDebuggerLocalController.cpp
--- a/Engine/Source/Developer/GameplayDebugger/Private/GameplayDebuggerLocalController.cpp
+++ b/Engine/Source/Developer/GameplayDebugger/Private/GameplayDebuggerLocalController.cpp
@@ -699,20 +699,19 @@ void UGameplayDebuggerLocalController::OnSelectionChanged(UObject* Object)
 	USelection* Selection = Cast<USelection>(Object);
 	if (Selection && CachedReplicator)
 	{
-		APawn* SelectedPawn = nullptr;
+		AActor* SelectedActor = nullptr;
 		for (int32 Idx = 0; Idx < Selection->Num(); Idx++)
 		{
-			AController* SelectedController = Cast<AController>(Selection->GetSelectedObject(Idx));
-			SelectedPawn = SelectedController ? SelectedController->GetPawn() : Cast<APawn>(Selection->GetSelectedObject(Idx));
-			if (SelectedPawn)
+			SelectedActor = Cast<AActor>(Selection->GetSelectedObject(Idx));
+			if (SelectedActor)
 			{
 				break;
 			}
 		}
 
-		if (SelectedPawn)
+		if (SelectedActor)
 		{
-			CachedReplicator->SetDebugActor(SelectedPawn, false);
+			CachedReplicator->SetDebugActor(SelectedActor, false);
 			CachedReplicator->CollectCategoryData(/*bForce=*/true);
 		}
 	}
