PR #6657: Fix a minor typo in user-facing text in the Gameplay Debugger UI (Contributed by [REDACTED_PERSON])

diff --git a/Engine/Source/Developer/GameplayDebugger/Private/GameplayDebuggerLocalController.cpp b/Engine/Source/Developer/GameplayDebugger/Private/GameplayDebuggerLocalController.cpp
--- a/Engine/Source/Developer/GameplayDebugger/Private/GameplayDebuggerLocalController.cpp
+++ b/Engine/Source/Developer/GameplayDebugger/Private/GameplayDebuggerLocalController.cpp
@@ -195,7 +195,7 @@ void UGameplayDebuggerLocalController::DrawHeader(FGameplayDebuggerCanvasContext
 	CanvasContext.CursorY = UsePaddingTop;
 	if (bSimulateMode)
 	{
-		CanvasContext.Printf(TEXT("Clear {yellow}DebugAI{white} show flag to close, use %s to toggle catories."), *CategoryKeysDesc);
+		CanvasContext.Printf(TEXT("Clear {yellow}DebugAI{white} show flag to close, use %s to toggle categories."), *CategoryKeysDesc);
 
 		// reactivate editor mode when this is being drawn = show flag is set
 #if WITH_EDITOR
@@ -204,7 +204,7 @@ void UGameplayDebuggerLocalController::DrawHeader(FGameplayDebuggerCanvasContext
 	}
 	else
 	{
-		CanvasContext.Printf(TEXT("Tap {yellow}%s{white} to close, use %s to toggle catories."), *ActivationKeyDesc, *CategoryKeysDesc);
+		CanvasContext.Printf(TEXT("Tap {yellow}%s{white} to close, use %s to toggle categories."), *ActivationKeyDesc, *CategoryKeysDesc);
 	}
 
 	const FString DebugActorDesc = FString::Printf(TEXT("Debug actor: {cyan}%s"), *CachedReplicator->GetDebugActorName().ToString());
