Made it possible to summon in-game hud for GameplayDebugger during replay playback

Yoan.StAmant, Maxime.Mercier

diff --git a/Engine/Source/Developer/GameplayDebugger/Private/GameplayDebuggerAddonManager.cpp b/Engine/Source/Developer/GameplayDebugger/Private/GameplayDebuggerAddonManager.cpp
--- a/Engine/Source/Developer/GameplayDebugger/Private/GameplayDebuggerAddonManager.cpp
+++ b/Engine/Source/Developer/GameplayDebugger/Private/GameplayDebuggerAddonManager.cpp
@@ -120,8 +120,9 @@ void FGameplayDebuggerAddonManager::NotifyCategoriesChanged()
 void FGameplayDebuggerAddonManager::CreateCategories(AGameplayDebuggerCategoryReplicator& Owner, TArray<TSharedRef<FGameplayDebuggerCategory> >& CategoryObjects)
 {
 	UWorld* World = Owner.GetWorld();
+	check(World);
 	const ENetMode NetMode = World->GetNetMode();
-	const bool bHasAuthority = (NetMode != NM_Client);
+	const bool bHasAuthority = FGameplayDebuggerUtils::IsAuthority(World);
 	const bool bIsLocal = (NetMode != NM_DedicatedServer);
 	const bool bIsSimulate = FGameplayDebuggerAddonBase::IsSimulateInEditor();
 
diff --git a/Engine/Source/Developer/GameplayDebugger/Private/GameplayDebuggerCategoryReplicator.cpp b/Engine/Source/Developer/GameplayDebugger/Private/GameplayDebuggerCategoryReplicator.cpp
--- a/Engine/Source/Developer/GameplayDebugger/Private/GameplayDebuggerCategoryReplicator.cpp
+++ b/Engine/Source/Developer/GameplayDebugger/Private/GameplayDebuggerCategoryReplicator.cpp
@@ -418,8 +418,9 @@ void AGameplayDebuggerCategoryReplicator::BeginPlay()
 	Super::BeginPlay();
 
 	UWorld* World = GetWorld();
+	check(World);
 	const ENetMode NetMode = World->GetNetMode();
-	bHasAuthority = (NetMode != NM_Client);
+	bHasAuthority = FGameplayDebuggerUtils::IsAuthority(World);
 	bIsLocal = (NetMode != NM_DedicatedServer);
 
 	FGameplayDebuggerAddonManager& AddonManager = FGameplayDebuggerAddonManager::GetCurrent();
@@ -722,9 +723,7 @@ void AGameplayDebuggerCategoryReplicator::SetReplicatorOwner(APlayerController*
 	{
 		// can't use bHasAuthority, BeginPlay was not called yet
 		UWorld* World = GetWorld();
-		const ENetMode NetMode = World->GetNetMode();
-
-		if (NetMode != NM_Client)
+		if (FGameplayDebuggerUtils::IsAuthority(World))
 		{
 			APlayerController* OldOwner = OwnerPC;
 			OwnerPC = InOwnerPC;
diff --git a/Engine/Source/Developer/GameplayDebugger/Private/GameplayDebuggerPlayerManager.cpp b/Engine/Source/Developer/GameplayDebugger/Private/GameplayDebuggerPlayerManager.cpp
--- a/Engine/Source/Developer/GameplayDebugger/Private/GameplayDebuggerPlayerManager.cpp
+++ b/Engine/Source/Developer/GameplayDebugger/Private/GameplayDebuggerPlayerManager.cpp
@@ -35,9 +35,10 @@ void AGameplayDebuggerPlayerManager::BeginPlay()
 	Super::BeginPlay();
 
 	UWorld* World = GetWorld();
+	check(World);
 	const ENetMode NetMode = World->GetNetMode();
 	
-	bHasAuthority = (NetMode != NM_Client);
+	bHasAuthority = FGameplayDebuggerUtils::IsAuthority(World);
 	bIsLocal = (NetMode != NM_DedicatedServer);
 	bInitialized = true;
 
diff --git a/Engine/Source/Developer/GameplayDebugger/Private/GameplayDebuggerTypes.cpp b/Engine/Source/Developer/GameplayDebugger/Private/GameplayDebuggerTypes.cpp
--- a/Engine/Source/Developer/GameplayDebugger/Private/GameplayDebuggerTypes.cpp
+++ b/Engine/Source/Developer/GameplayDebugger/Private/GameplayDebuggerTypes.cpp
@@ -8,9 +8,18 @@
 #include "DrawDebugHelpers.h"
 #include "CanvasItem.h"
 #include "Engine/Canvas.h"
+#include "Engine/World.h"
 
 DEFINE_LOG_CATEGORY(LogGameplayDebug);
 
+namespace FGameplayDebuggerUtils
+{
+	bool IsAuthority(UWorld* World)
+	{
+		return (World == nullptr) || (World->GetNetMode() != NM_Client) || World->IsPlayingReplay();
+	}
+}
+
 //////////////////////////////////////////////////////////////////////////
 // FGameplayDebuggerShape
 
diff --git a/Engine/Source/Developer/GameplayDebugger/Public/GameplayDebuggerTypes.h b/Engine/Source/Developer/GameplayDebugger/Public/GameplayDebuggerTypes.h
--- a/Engine/Source/Developer/GameplayDebugger/Public/GameplayDebuggerTypes.h
+++ b/Engine/Source/Developer/GameplayDebugger/Public/GameplayDebuggerTypes.h
@@ -15,6 +15,11 @@ DECLARE_LOG_CATEGORY_EXTERN(LogGameplayDebug, Log, All);
 class FCanvasItem;
 struct FCanvasIcon;
 
+namespace FGameplayDebuggerUtils
+{
+	bool GAMEPLAYDEBUGGER_API IsAuthority(UWorld* World);
+}
+
 class GAMEPLAYDEBUGGER_API FGameplayDebuggerCanvasContext
 {
 public:
