PR #6711: FGameplayDebuggerCanvasContext::PrintAt functions to not save and res. (Contributed by [REDACTED_PERSON])

diff --git a/Engine/Source/Developer/GameplayDebugger/Private/GameplayDebuggerTypes.cpp b/Engine/Source/Developer/GameplayDebugger/Private/GameplayDebuggerTypes.cpp
--- a/Engine/Source/Developer/GameplayDebugger/Private/GameplayDebuggerTypes.cpp
+++ b/Engine/Source/Developer/GameplayDebugger/Private/GameplayDebuggerTypes.cpp
@@ -415,32 +415,17 @@ void FGameplayDebuggerCanvasContext::Print(const FColor& Color, const FString& S
 
 void FGameplayDebuggerCanvasContext::PrintAt(float PosX, float PosY, const FString& String)
 {
-	const float SavedPosX = CursorX;
-	const float SavedPosY = CursorY;
-	const float SavedDefX = DefaultX;
-
-	DefaultX = CursorX = PosX;
-	DefaultY = CursorY = PosY;
-	Print(FColor::White, String);
-
-	CursorX = SavedPosX;
-	CursorY = SavedPosY;
-	DefaultX = SavedDefX;
+	PrintAt(PosX, PosY, FColor::White, String);
 }
 
 void FGameplayDebuggerCanvasContext::PrintAt(float PosX, float PosY, const FColor& Color, const FString& String)
 {
-	const float SavedPosX = CursorX;
-	const float SavedPosY = CursorY;
-	const float SavedDefX = DefaultX;
+	TGuardValue<float> ScopedCursorX(CursorX, PosX);
+	TGuardValue<float> ScopedCursorY(CursorY, PosY);
+	TGuardValue<float> ScopedDefaultX(DefaultX, PosX);
+	TGuardValue<float> ScopedDefaultY(DefaultY, PosY);
 
-	DefaultX = CursorX = PosX;
-	DefaultY = CursorY = PosY;
 	Print(Color, String);
-
-	CursorX = SavedPosX;
-	CursorY = SavedPosY;
-	DefaultX = SavedDefX;
 }
 
 // copied from Core/Private/Misc/VarargsHeler.h 
