Made displaying GameplayDebugger's UI possible in pure editor-mode

diff --git a/Engine/Source/Developer/GameplayDebugger/Private/GameplayDebuggerAddonBase.cpp b/Engine/Source/Developer/GameplayDebugger/Private/GameplayDebuggerAddonBase.cpp
--- a/Engine/Source/Developer/GameplayDebugger/Private/GameplayDebuggerAddonBase.cpp
+++ b/Engine/Source/Developer/GameplayDebugger/Private/GameplayDebuggerAddonBase.cpp
@@ -2,6 +2,7 @@
 
 #include "GameplayDebuggerAddonBase.h"
 #include "GameplayDebuggerCategoryReplicator.h"
+#include "GameFramework/PlayerController.h"
 #if WITH_EDITOR
 #include "Editor/EditorEngine.h"
 #endif
@@ -22,6 +23,27 @@ FString FGameplayDebuggerAddonBase::GetInputHandlerDescription(int32 HandlerId)
 	return InputHandlers.IsValidIndex(HandlerId) ? InputHandlers[HandlerId].ToString() : FString();
 }
 
+UWorld* FGameplayDebuggerAddonBase::GetWorldFromReplicator() const
+{
+	AGameplayDebuggerCategoryReplicator* Replicator = RepOwner.Get();
+	return Replicator ? Replicator->GetWorld() : nullptr;
+}
+
+UWorld* FGameplayDebuggerAddonBase::GetDataWorld(const APlayerController* OwnerPC, const AActor* DebugActor) const
+{
+	// We're trying OwnerPC first because that's the main owner for what's being displayed by gameplay debugger
+	// but it editor-only mode there are no player controllers, so we try the DebugActor, which can be picked
+	// also in the editor-mode by selecting an actor on the level.
+	// finally, DebugActor can be null as well, when no actor is selected, when we fetch the world which the
+	// AGameplayDebuggerCategoryReplicator instance owning this debugger addon belongs to.
+	UWorld* World = OwnerPC ? OwnerPC->GetWorld() : nullptr;
+	if (World == nullptr && DebugActor != nullptr)
+	{
+		World = DebugActor->GetWorld();
+	}
+	return World ? World : GetWorldFromReplicator();
+}
+
 void FGameplayDebuggerAddonBase::OnGameplayDebuggerActivated()
 {
 	// empty in base class
diff --git a/Engine/Source/Developer/GameplayDebugger/Private/GameplayDebuggerCategory.cpp b/Engine/Source/Developer/GameplayDebugger/Private/GameplayDebuggerCategory.cpp
--- a/Engine/Source/Developer/GameplayDebugger/Private/GameplayDebuggerCategory.cpp
+++ b/Engine/Source/Developer/GameplayDebugger/Private/GameplayDebuggerCategory.cpp
@@ -63,7 +63,13 @@ void FGameplayDebuggerCategory::AddShape(const FGameplayDebuggerShape& Shape)
 
 void FGameplayDebuggerCategory::DrawCategory(APlayerController* OwnerPC, FGameplayDebuggerCanvasContext& CanvasContext)
 {
-	UWorld* World = OwnerPC->GetWorld();
+	// note that we prefer OwnerPC's world here since if it's given then that's the end user of the data
+	UWorld* World = OwnerPC ? OwnerPC->GetWorld() : CanvasContext.World.Get();
+
+	if (!ensure(World))
+	{
+		return;
+	}
 
 	FString CategoryPrefix;
 	if (!bShowCategoryName)
diff --git a/Engine/Source/Developer/GameplayDebugger/Private/GameplayDebuggerCategoryReplicator.cpp b/Engine/Source/Developer/GameplayDebugger/Private/GameplayDebuggerCategoryReplicator.cpp
--- a/Engine/Source/Developer/GameplayDebugger/Private/GameplayDebuggerCategoryReplicator.cpp
+++ b/Engine/Source/Developer/GameplayDebugger/Private/GameplayDebuggerCategoryReplicator.cpp
@@ -408,6 +408,9 @@ AGameplayDebuggerCategoryReplicator::AGameplayDebuggerCategoryReplicator(const F
 
 	bIsEnabled = false;
 	bIsEnabledLocal = false;
+	bHasAuthority = false;
+	bIsLocal = false;
+	bIsEditorWorldReplicator = false;
 	bReplicates = true;
 
 	ReplicatedData.Owner = this;
@@ -423,6 +426,24 @@ void AGameplayDebuggerCategoryReplicator::BeginPlay()
 	bHasAuthority = FGameplayDebuggerUtils::IsAuthority(World);
 	bIsLocal = (NetMode != NM_DedicatedServer);
 
+	Init();
+}
+
+#if WITH_EDITOR
+void AGameplayDebuggerCategoryReplicator::InitForEditor()
+{
+	bHasAuthority = true;
+	bIsLocal = true;
+	bIsEditorWorldReplicator = true;
+	Init();
+}
+#endif // WITH_EDITOR
+
+void AGameplayDebuggerCategoryReplicator::Init()
+{
+	UWorld* World = GetWorld();
+	check(World);
+
 	FGameplayDebuggerAddonManager& AddonManager = FGameplayDebuggerAddonManager::GetCurrent();
 	AddonManager.OnCategoriesChanged.AddUObject(this, &AGameplayDebuggerCategoryReplicator::OnCategoriesChanged);
 	AddonManager.OnExtensionsChanged.AddUObject(this, &AGameplayDebuggerCategoryReplicator::OnExtensionsChanged);
@@ -430,9 +451,6 @@ void AGameplayDebuggerCategoryReplicator::BeginPlay()
 	OnCategoriesChanged();
 	OnExtensionsChanged();
 
-	AGameplayDebuggerPlayerManager& PlayerManager = AGameplayDebuggerPlayerManager::GetCurrent(GetWorld());
-	PlayerManager.RegisterReplicator(*this);
-
 	SetActorHiddenInGame(!bIsLocal);
 	if (bIsLocal)
 	{
@@ -445,6 +463,9 @@ void AGameplayDebuggerCategoryReplicator::BeginPlay()
 	{
 		SetEnabled(FGameplayDebuggerAddonBase::IsSimulateInEditor());
 	}
+
+	AGameplayDebuggerPlayerManager& PlayerManager = AGameplayDebuggerPlayerManager::GetCurrent(World);
+	PlayerManager.RegisterReplicator(*this);
 }
 
 void AGameplayDebuggerCategoryReplicator::EndPlay(const EEndPlayReason::Type EndPlayReason)
@@ -626,7 +647,7 @@ void AGameplayDebuggerCategoryReplicator::OnReceivedDataPackPacket(int32 Categor
 void AGameplayDebuggerCategoryReplicator::TickActor(float DeltaTime, enum ELevelTick TickType, FActorTickFunction& ThisTickFunction)
 {
 	Super::TickActor(DeltaTime, TickType, ThisTickFunction);
-	if (OwnerPC)
+	if (OwnerPC || bIsEditorWorldReplicator)
 	{
 		CollectCategoryData();
 	}
diff --git a/Engine/Source/Developer/GameplayDebugger/Private/GameplayDebuggerLocalController.cpp b/Engine/Source/Developer/GameplayDebugger/Private/GameplayDebuggerLocalController.cpp
--- a/Engine/Source/Developer/GameplayDebugger/Private/GameplayDebuggerLocalController.cpp
+++ b/Engine/Source/Developer/GameplayDebugger/Private/GameplayDebuggerLocalController.cpp
@@ -30,6 +30,7 @@
 #if WITH_EDITOR
 #include "Editor/GameplayDebuggerEdMode.h"
 #include "EditorModeManager.h"
+#include "Editor.h"
 #endif // WITH_EDITOR
 
 UGameplayDebuggerLocalController::UGameplayDebuggerLocalController(const FObjectInitializer& ObjectInitializer) : Super(ObjectInitializer)
@@ -41,13 +42,16 @@ UGameplayDebuggerLocalController::UGameplayDebuggerLocalController(const FObject
 	bPrevLocallyEnabled = false;
 	bEnableTextShadow = false;
 	ActiveRowIdx = 0;
+#if WITH_EDITOR
+	bActivateOnPIEEnd = false;
+#endif // WITH_EDITOR
 }
 
 void UGameplayDebuggerLocalController::Initialize(AGameplayDebuggerCategoryReplicator& Replicator, AGameplayDebuggerPlayerManager& Manager)
 {
 	CachedReplicator = &Replicator;
 	CachedPlayerManager = &Manager;
-	bSimulateMode = FGameplayDebuggerAddonBase::IsSimulateInEditor();
+	bSimulateMode = FGameplayDebuggerAddonBase::IsSimulateInEditor() || Replicator.IsEditorWorldReplicator();
 
 	UDebugDrawService::Register(bSimulateMode ? TEXT("DebugAI") : TEXT("Game"), FDebugDrawDelegate::CreateUObject(this, &UGameplayDebuggerLocalController::OnDebugDraw));
 
@@ -56,6 +60,13 @@ void UGameplayDebuggerLocalController::Initialize(AGameplayDebuggerCategoryRepli
 	{
 		USelection::SelectionChangedEvent.AddUObject(this, &UGameplayDebuggerLocalController::OnSelectionChanged);
 		USelection::SelectObjectEvent.AddUObject(this, &UGameplayDebuggerLocalController::OnSelectedObject);
+
+		if (Replicator.IsEditorWorldReplicator())
+		{
+			// bind to PIE start and end notifies to hide before pie and re-enable if need be when pie's done
+			FEditorDelegates::BeginPIE.AddUObject(this, &UGameplayDebuggerLocalController::OnBeginPIE);
+			FEditorDelegates::EndPIE.AddUObject(this, &UGameplayDebuggerLocalController::OnEndPIE);
+		}
 	}
 #endif
 
@@ -138,6 +149,9 @@ void UGameplayDebuggerLocalController::OnDebugDraw(class UCanvas* Canvas, class
 
 		CanvasContext.FontRenderInfo.bEnableShadow = bEnableTextShadow;
 
+		CanvasContext.PlayerController = CachedReplicator->GetReplicationOwner();
+		CanvasContext.World = CachedReplicator->GetWorld();
+
 		DrawHeader(CanvasContext);
 
 		if (DataPackMap.Num() != NumCategories)
@@ -806,6 +820,28 @@ void UGameplayDebuggerLocalController::RebuildDataPackMap()
 	}
 }
 
+#if WITH_EDITOR
+void UGameplayDebuggerLocalController::OnBeginPIE(const bool bIsSimulating)
+{
+	bActivateOnPIEEnd = bIsLocallyEnabled;
+	if (bIsLocallyEnabled)
+	{
+		ToggleActivation();
+	}
+}
+
+void UGameplayDebuggerLocalController::OnEndPIE(const bool bIsSimulating)
+{
+	if (bActivateOnPIEEnd && !bIsLocallyEnabled)
+	{
+		ToggleActivation();
+	}
+}
+#endif // WITH_EDITOR
+
+//----------------------------------------------------------------------//
+// FGameplayDebuggerConsoleCommands 
+//----------------------------------------------------------------------//
 /**  Helper structure to declare/define console commands in the source file and to access UGameplayDebuggerLocalController protected members */
 struct FGameplayDebuggerConsoleCommands
 {
@@ -819,6 +855,12 @@ private:
 		{
 			Controller = AGameplayDebuggerPlayerManager::GetCurrent(InWorld).GetLocalController(*LocalPC);
 		}
+#if WITH_EDITOR
+		else if (InWorld != nullptr && InWorld->IsGameWorld() == false)
+		{
+			Controller = AGameplayDebuggerPlayerManager::GetCurrent(InWorld).GetEditorController();
+		}
+#endif // WITH_EDITOR
 
 		UE_CLOG(Controller == nullptr, LogConsoleResponse, Error, TEXT("GameplayDebugger not available"));
 		return Controller;
diff --git a/Engine/Source/Developer/GameplayDebugger/Private/GameplayDebuggerModule.cpp b/Engine/Source/Developer/GameplayDebugger/Private/GameplayDebuggerModule.cpp
--- a/Engine/Source/Developer/GameplayDebugger/Private/GameplayDebuggerModule.cpp
+++ b/Engine/Source/Developer/GameplayDebugger/Private/GameplayDebuggerModule.cpp
@@ -168,6 +168,7 @@ AGameplayDebuggerPlayerManager& FGameplayDebuggerModule::GetPlayerManager(UWorld
 	{
 		ManagerOb = World->SpawnActor<AGameplayDebuggerPlayerManager>();
 		PlayerManagers.Add(World, ManagerOb);
+		ManagerOb->Init();
 	}
 
 	check(ManagerOb);
@@ -177,7 +178,7 @@ AGameplayDebuggerPlayerManager& FGameplayDebuggerModule::GetPlayerManager(UWorld
 void FGameplayDebuggerModule::OnWorldInitialized(UWorld* World, const UWorld::InitializationValues IVS)
 {
 	// make sure that world has valid player manager, create when it doesn't
-	if (World && World->IsGameWorld())
+	if (World && (World->IsGameWorld() || World->IsEditorWorld()))
 	{
 		GetPlayerManager(World);
 	}
diff --git a/Engine/Source/Developer/GameplayDebugger/Private/GameplayDebuggerPlayerManager.cpp b/Engine/Source/Developer/GameplayDebugger/Private/GameplayDebuggerPlayerManager.cpp
--- a/Engine/Source/Developer/GameplayDebugger/Private/GameplayDebuggerPlayerManager.cpp
+++ b/Engine/Source/Developer/GameplayDebugger/Private/GameplayDebuggerPlayerManager.cpp
@@ -6,6 +6,8 @@
 #include "GameplayDebuggerCategoryReplicator.h"
 #include "GameplayDebuggerLocalController.h"
 #include "Engine/DebugCameraController.h"
+#include "Engine/InputDelegateBinding.h"
+
 
 AGameplayDebuggerPlayerManager::AGameplayDebuggerPlayerManager(const FObjectInitializer& ObjectInitializer) : Super(ObjectInitializer)
 {
@@ -30,6 +32,26 @@ AGameplayDebuggerPlayerManager::AGameplayDebuggerPlayerManager(const FObjectInit
 	bInitialized = false;
 }
 
+void AGameplayDebuggerPlayerManager::PostInitProperties()
+{
+	Super::PostInitProperties();
+
+	if (HasAnyFlags(RF_ClassDefaultObject) == false)
+	{
+		if (InputComponent == nullptr)
+		{
+			// create an InputComponent object so that the level script actor can bind key events
+			InputComponent = NewObject<UInputComponent>(this, TEXT("GameplayDebug_Input"));
+			InputComponent->RegisterComponent();
+		}
+
+		if (UInputDelegateBinding::SupportsInputDelegate(GetClass()))
+		{
+			UInputDelegateBinding::BindInputDelegates(GetClass(), InputComponent);
+		}
+	}
+}
+
 void AGameplayDebuggerPlayerManager::BeginPlay()
 {
 	Super::BeginPlay();
@@ -71,6 +93,32 @@ void AGameplayDebuggerPlayerManager::EndPlay(const EEndPlayReason::Type Reason)
 	}
 }
 
+void AGameplayDebuggerPlayerManager::Init()
+{
+#if WITH_EDITOR
+	UWorld* World = GetWorld();
+	if (World != nullptr && World->WorldType == EWorldType::Editor)
+	{
+		bHasAuthority = true;
+		bIsLocal = true;
+		bInitialized = true;
+
+		AGameplayDebuggerCategoryReplicator* Replicator = World->SpawnActorDeferred<AGameplayDebuggerCategoryReplicator>(AGameplayDebuggerCategoryReplicator::StaticClass(), FTransform::Identity);
+		Replicator->SetReplicatorOwner(nullptr);
+		Replicator->FinishSpawning(FTransform::Identity, true);
+		SetActorTickEnabled(true);
+
+		EditorWorldData.Replicator = Replicator;
+
+		Replicator->InitForEditor();
+
+		EditorWorldData.Controller = NewObject<UGameplayDebuggerLocalController>(this, TEXT("GameplayDebug_Controller_Editor"));
+		EditorWorldData.Controller->Initialize(*Replicator, *this);
+		
+	}
+#endif // WITH_EDITOR
+}
+
 void AGameplayDebuggerPlayerManager::TickActor(float DeltaTime, enum ELevelTick TickType, FActorTickFunction& ThisTickFunction)
 {
 	Super::TickActor(DeltaTime, TickType, ThisTickFunction);
@@ -120,12 +168,6 @@ void AGameplayDebuggerPlayerManager::UpdateAuthReplicators()
 
 void AGameplayDebuggerPlayerManager::RegisterReplicator(AGameplayDebuggerCategoryReplicator& Replicator)
 {
-	APlayerController* OwnerPC = Replicator.GetReplicationOwner();
-	if (OwnerPC == nullptr)
-	{
-		return;
-	}
-
 	if (!bInitialized)
 	{
 		PendingRegistrations.Add(&Replicator);
@@ -138,14 +180,19 @@ void AGameplayDebuggerPlayerManager::RegisterReplicator(AGameplayDebuggerCategor
 	
 	if (bIsLocal)
 	{
-		NewData.InputComponent = NewObject<UInputComponent>(OwnerPC, TEXT("GameplayDebug_Input"));
-		NewData.InputComponent->Priority = -1;
+		APlayerController* OwnerPC = Replicator.GetReplicationOwner();
 
-		NewData.Controller = NewObject<UGameplayDebuggerLocalController>(OwnerPC, TEXT("GameplayDebug_Controller"));
+		NewData.InputComponent = OwnerPC ? NewObject<UInputComponent>(OwnerPC, TEXT("GameplayDebug_Input")) : InputComponent;
+		check(NewData.InputComponent);
+		NewData.InputComponent->Priority = -1;
+		NewData.Controller = NewObject<UGameplayDebuggerLocalController>(OwnerPC ? OwnerPC : (AActor*)this, TEXT("GameplayDebug_Controller"));
 		NewData.Controller->Initialize(Replicator, *this);
 		NewData.Controller->BindInput(*NewData.InputComponent);
 
-		OwnerPC->PushInputComponent(NewData.InputComponent);
+		if (OwnerPC)
+		{
+			OwnerPC->PushInputComponent(NewData.InputComponent);
+		}
 	}
 	else
 	{
@@ -203,3 +250,34 @@ const FGameplayDebuggerPlayerData* AGameplayDebuggerPlayerManager::GetPlayerData
 
 	return nullptr;
 }
+
+#if WITH_EDITOR
+// FTickableGameObject begin
+void AGameplayDebuggerPlayerManager::Tick(float DeltaTime)
+{
+#if WITH_EDITORONLY_DATA 
+	if (EditorWorldData.Replicator)
+	{
+		FActorTickFunction DummyTickFunction;
+		EditorWorldData.Replicator->TickActor(DeltaTime, ELevelTick::LEVELTICK_All, DummyTickFunction);
+	}
+#endif // WITH_EDITORONLY_DATA 
+}
+
+ETickableTickType AGameplayDebuggerPlayerManager::GetTickableTickType() const
+{
+	return IsTickable() ? ETickableTickType::Conditional : ETickableTickType::Never;
+}
+
+TStatId AGameplayDebuggerPlayerManager::GetStatId() const
+{
+	RETURN_QUICK_DECLARE_CYCLE_STAT(AGameplayDebuggerPlayerManager, STATGROUP_Tickables);
+}
+
+bool AGameplayDebuggerPlayerManager::IsTickable() const
+{
+	UWorld* World = GetWorld();
+	return (World != nullptr) && (World->IsEditorWorld() == true) && (World->IsGameWorld() == false);
+}
+// FTickableGameObject end
+#endif // WITH_EDITOR
diff --git a/Engine/Source/Developer/GameplayDebugger/Private/GameplayDebuggerTypes.cpp b/Engine/Source/Developer/GameplayDebugger/Private/GameplayDebuggerTypes.cpp
--- a/Engine/Source/Developer/GameplayDebugger/Private/GameplayDebuggerTypes.cpp
+++ b/Engine/Source/Developer/GameplayDebugger/Private/GameplayDebuggerTypes.cpp
@@ -9,6 +9,7 @@
 #include "CanvasItem.h"
 #include "Engine/Canvas.h"
 #include "Engine/World.h"
+#include "GameFramework/PlayerController.h"
 
 DEFINE_LOG_CATEGORY(LogGameplayDebug);
 
@@ -556,6 +557,13 @@ void FGameplayDebuggerCanvasContext::DrawIcon(const FColor& Color, const FCanvas
 	}
 }
 
+UWorld* FGameplayDebuggerCanvasContext::GetWorld() const 
+{
+	return World.IsValid() 
+		? World.Get()
+		: (PlayerController.IsValid() ? PlayerController->GetWorld() : nullptr);
+}
+
 //////////////////////////////////////////////////////////////////////////
 // FGameplayDebuggerDataPack
 
diff --git a/Engine/Source/Developer/GameplayDebugger/Public/GameplayDebuggerAddonBase.h b/Engine/Source/Developer/GameplayDebugger/Public/GameplayDebuggerAddonBase.h
--- a/Engine/Source/Developer/GameplayDebugger/Public/GameplayDebuggerAddonBase.h
+++ b/Engine/Source/Developer/GameplayDebugger/Public/GameplayDebuggerAddonBase.h
@@ -5,7 +5,9 @@
 #include "CoreMinimal.h"
 #include "GameplayDebuggerTypes.h"
 
+class UWorld;
 class AActor;
+class APlayerController;
 class AGameplayDebuggerCategoryReplicator;
 
 class GAMEPLAYDEBUGGER_API FGameplayDebuggerAddonBase : public TSharedFromThis<FGameplayDebuggerAddonBase>
@@ -17,6 +19,9 @@ public:
 	int32 GetNumInputHandlers() const { return InputHandlers.Num(); }
 	FGameplayDebuggerInputHandler& GetInputHandler(int32 HandlerId) { return InputHandlers[HandlerId]; }
 	FString GetInputHandlerDescription(int32 HandlerId) const;
+	UWorld* GetWorldFromReplicator() const;
+	/** Returns the first non-null world associated with, in order: OwnerPC, DebugActor, AGameplayDebuggerCategoryReplicator owner */
+	UWorld* GetDataWorld(const APlayerController* OwnerPC, const AActor* DebugActor) const;
 
 	/** [ALL] called when gameplay debugger is activated */
 	virtual void OnGameplayDebuggerActivated();
diff --git a/Engine/Source/Developer/GameplayDebugger/Public/GameplayDebuggerCategoryReplicator.h b/Engine/Source/Developer/GameplayDebugger/Public/GameplayDebuggerCategoryReplicator.h
--- a/Engine/Source/Developer/GameplayDebugger/Public/GameplayDebuggerCategoryReplicator.h
+++ b/Engine/Source/Developer/GameplayDebugger/Public/GameplayDebuggerCategoryReplicator.h
@@ -127,6 +127,9 @@ public:
 	/** get replicator state */
 	bool IsEnabled() const { return bIsEnabled; }
 
+	/** returns true if this replicator has been created for an editor world */
+	bool IsEditorWorldReplicator() const { return bIsEditorWorldReplicator; }
+
 	/** get category state */
 	bool IsCategoryEnabled(int32 CategoryId) const;
 
@@ -148,6 +151,10 @@ public:
 	/** returns true if object was created for local player (client / standalone) */
 	bool IsLocal() const { return bIsLocal; }
 
+#if WITH_EDITOR
+	void InitForEditor();
+#endif // WITH_EDITOR
+
 protected:
 
 	friend FGameplayDebuggerNetPack;
@@ -180,6 +187,7 @@ protected:
 	uint32 bIsEnabledLocal : 1;
 	uint32 bHasAuthority : 1;
 	uint32 bIsLocal : 1;
+	uint32 bIsEditorWorldReplicator : 1;
 
 	/** notify about changes in known category set */
 	void OnCategoriesChanged();
@@ -212,4 +220,7 @@ protected:
 
 	/** [LOCAL] notify from CategoryData replication */
 	void OnReceivedDataPackPacket(int32 CategoryId, int32 DataPackId, const FGameplayDebuggerDataPack& DataPacket);
+
+	/** called both from BeginPlay and InitForEditor to setup instance's internal */
+	void Init();
 };
diff --git a/Engine/Source/Developer/GameplayDebugger/Public/GameplayDebuggerLocalController.h b/Engine/Source/Developer/GameplayDebugger/Public/GameplayDebuggerLocalController.h
--- a/Engine/Source/Developer/GameplayDebugger/Public/GameplayDebuggerLocalController.h
+++ b/Engine/Source/Developer/GameplayDebugger/Public/GameplayDebuggerLocalController.h
@@ -62,6 +62,9 @@ protected:
 	uint32 bIsLocallyEnabled : 1;
 	uint32 bPrevLocallyEnabled : 1;
 	uint32 bEnableTextShadow : 1;
+#if WITH_EDITOR
+	uint32 bActivateOnPIEEnd : 1;
+#endif // WITH_EDITOR
 
 	FString ActivationKeyDesc;
 	FString RowUpKeyDesc;
@@ -120,6 +123,8 @@ protected:
 	/** event for simulate in editor mode */
 	void OnSelectionChanged(UObject* Object);
 	void OnSelectedObject(UObject* Object);
+	void OnBeginPIE(const bool bIsSimulating);
+	void OnEndPIE(const bool bIsSimulating);
 #endif
 
 	FString GetKeyDescriptionShort(const FKey& KeyBind) const;
diff --git a/Engine/Source/Developer/GameplayDebugger/Public/GameplayDebuggerPlayerManager.h b/Engine/Source/Developer/GameplayDebugger/Public/GameplayDebuggerPlayerManager.h
--- a/Engine/Source/Developer/GameplayDebugger/Public/GameplayDebuggerPlayerManager.h
+++ b/Engine/Source/Developer/GameplayDebugger/Public/GameplayDebuggerPlayerManager.h
@@ -5,6 +5,7 @@
 #include "CoreMinimal.h"
 #include "UObject/ObjectMacros.h"
 #include "GameFramework/Actor.h"
+#include "Tickable.h"
 #include "GameplayDebuggerPlayerManager.generated.h"
 
 class AGameplayDebuggerCategoryReplicator;
@@ -28,17 +29,32 @@ struct FGameplayDebuggerPlayerData
 };
 
 UCLASS(NotBlueprintable, NotBlueprintType, notplaceable, noteditinlinenew, hidedropdown, Transient)
-class GAMEPLAYDEBUGGER_API AGameplayDebuggerPlayerManager : public AActor
+class GAMEPLAYDEBUGGER_API AGameplayDebuggerPlayerManager : public AActor, public FTickableGameObject
 {
 	GENERATED_UCLASS_BODY()
 
 	virtual void TickActor(float DeltaTime, enum ELevelTick TickType, FActorTickFunction& ThisTickFunction) override;
 
+#if WITH_EDITOR
+	// we're ticking only the manager only when in editor
+	// FTickableGameObject begin
+	virtual UWorld* GetTickableGameObjectWorld() const override { return GetWorld(); }
+	virtual void Tick(float DeltaTime) override;
+	virtual ETickableTickType GetTickableTickType() const override;
+	virtual TStatId GetStatId() const override;
+	virtual bool IsTickableInEditor() const override { return true; }
+	virtual bool IsTickable() const override;
+	// FTickableGameObject end
+#endif // WITH_EDITOR
+
+	virtual void PostInitProperties() override;
+
 protected:
 	virtual void BeginPlay() override;
 
 public:
 	virtual void EndPlay(const EEndPlayReason::Type Reason) override;
+	void Init();
 		
 	void UpdateAuthReplicators();
 	void RegisterReplicator(AGameplayDebuggerCategoryReplicator& Replicator);
@@ -47,6 +63,9 @@ public:
 	AGameplayDebuggerCategoryReplicator* GetReplicator(const APlayerController& OwnerPC) const;
 	UInputComponent* GetInputComponent(const APlayerController& OwnerPC) const;
 	UGameplayDebuggerLocalController* GetLocalController(const APlayerController& OwnerPC) const;
+#if WITH_EDITOR
+	UGameplayDebuggerLocalController* GetEditorController() const { return EditorWorldData.Controller; }
+#endif // WITH_EDITOR
 	
 	const FGameplayDebuggerPlayerData* GetPlayerData(const APlayerController& OwnerPC) const;
 
@@ -60,6 +79,11 @@ protected:
 	UPROPERTY()
 	TArray<AGameplayDebuggerCategoryReplicator*> PendingRegistrations;
 
+#if WITH_EDITORONLY_DATA 
+	UPROPERTY()
+	FGameplayDebuggerPlayerData EditorWorldData;
+#endif // WITH_EDITORONLY_DATA
+
 	uint32 bHasAuthority : 1;
 	uint32 bIsLocal : 1;
 	uint32 bInitialized : 1;
diff --git a/Engine/Source/Developer/GameplayDebugger/Public/GameplayDebuggerTypes.h b/Engine/Source/Developer/GameplayDebugger/Public/GameplayDebuggerTypes.h
--- a/Engine/Source/Developer/GameplayDebugger/Public/GameplayDebuggerTypes.h
+++ b/Engine/Source/Developer/GameplayDebugger/Public/GameplayDebuggerTypes.h
@@ -9,6 +9,8 @@ class FCanvasItem;
 class UCanvas;
 class UFont;
 struct FCanvasIcon;
+class APlayerController;
+class UWorld;
 
 DECLARE_LOG_CATEGORY_EXTERN(LogGameplayDebug, Log, All);
 
@@ -29,6 +31,12 @@ public:
 	/** current text font */
 	TWeakObjectPtr<UFont> Font;
 
+	/** the player controller associated with this debugger context  */
+	TWeakObjectPtr<APlayerController> PlayerController;
+
+	/** the world associated with this debugger context  */
+	TWeakObjectPtr<UWorld> World;
+
 	/** font render data */
 	FFontRenderInfo FontRenderInfo;
 
@@ -108,6 +116,10 @@ public:
 
 	// draw icon on canvas
 	void DrawIcon(const FColor& Color, const FCanvasIcon& Icon, float PosX, float PosY, float Scale = 1.f);
+
+	// fetches the World associated with this context. Will use the World member variable, if set, and 
+	// PlayerController's world otherwise. Note that it can still be null and needs to be tested.
+	UWorld* GetWorld() const;
 };
 
 namespace FGameplayDebuggerCanvasStrings
diff --git a/Engine/Source/Runtime/AIModule/Private/GameplayDebugger/GameplayDebuggerCategory_AI.cpp b/Engine/Source/Runtime/AIModule/Private/GameplayDebugger/GameplayDebuggerCategory_AI.cpp
--- a/Engine/Source/Runtime/AIModule/Private/GameplayDebugger/GameplayDebuggerCategory_AI.cpp
+++ b/Engine/Source/Runtime/AIModule/Private/GameplayDebugger/GameplayDebuggerCategory_AI.cpp
@@ -301,13 +301,13 @@ void FGameplayDebuggerCategory_AI::OnDataPackReplicated(int32 DataPackId)
 
 void FGameplayDebuggerCategory_AI::DrawData(APlayerController* OwnerPC, FGameplayDebuggerCanvasContext& CanvasContext)
 {
-	UWorld* MyWorld = OwnerPC->GetWorld();
+	UWorld* MyWorld = CanvasContext.GetWorld();
 	AActor* SelectedActor = FindLocalDebugActor();
 
 	const bool bReducedMode = IsSimulateInEditor();
 	bShowCategoryName = !bReducedMode || DataPack.bHasController;
 
-	if (FGameplayDebuggerCategoryTweakables::bDrawOverheadIcons)
+	if (FGameplayDebuggerCategoryTweakables::bDrawOverheadIcons && OwnerPC)
 	{
 		DrawPawnIcons(MyWorld, SelectedActor, OwnerPC->GetPawn(), CanvasContext);
 	}
diff --git a/Engine/Source/Runtime/AIModule/Private/GameplayDebugger/GameplayDebuggerCategory_EQS.cpp b/Engine/Source/Runtime/AIModule/Private/GameplayDebugger/GameplayDebuggerCategory_EQS.cpp
--- a/Engine/Source/Runtime/AIModule/Private/GameplayDebugger/GameplayDebuggerCategory_EQS.cpp
+++ b/Engine/Source/Runtime/AIModule/Private/GameplayDebugger/GameplayDebuggerCategory_EQS.cpp
@@ -121,6 +121,11 @@ FDebugRenderSceneProxy* FGameplayDebuggerCategory_EQS::CreateDebugSceneProxy(con
 void FGameplayDebuggerCategory_EQS::DrawData(APlayerController* OwnerPC, FGameplayDebuggerCanvasContext& CanvasContext)
 {
 #if USE_EQS_DEBUGGER
+	UWorld* World = CanvasContext.GetWorld();
+	if (World == nullptr)
+	{
+		return;
+	}
 
 	const FString HeaderDesc = (DataPack.QueryDebugData.Num() > 1) ?
 		FString::Printf(TEXT("Queries: {yellow}%d{white}, press {yellow}[%s]{white} to cycle through"), DataPack.QueryDebugData.Num(), *GetInputHandlerDescription(0)) :
@@ -146,7 +151,7 @@ void FGameplayDebuggerCategory_EQS::DrawData(APlayerController* OwnerPC, FGamepl
 		const EQSDebug::FQueryData& ShownQueryData = DataPack.QueryDebugData[ShownQueryIndex];
 
 		CanvasContext.MoveToNewLine();
-		CanvasContext.Printf(TEXT("Timestamp: {yellow}%.3f (~ %.2fs ago)"), ShownQueryData.Timestamp, OwnerPC->GetWorld()->TimeSince(ShownQueryData.Timestamp));
+		CanvasContext.Printf(TEXT("Timestamp: {yellow}%.3f (~ %.2fs ago)"), ShownQueryData.Timestamp, World->TimeSince(ShownQueryData.Timestamp));
 		
 		FString OptionsDesc(TEXT("Options: "));
 		for (int32 Idx = 0; Idx < ShownQueryData.Options.Num(); Idx++)
@@ -168,7 +173,8 @@ void FGameplayDebuggerCategory_EQS::DrawData(APlayerController* OwnerPC, FGamepl
 #if USE_EQS_DEBUGGER
 int32 FGameplayDebuggerCategory_EQS::DrawLookedAtItem(const EQSDebug::FQueryData& QueryData, APlayerController* OwnerPC, FGameplayDebuggerCanvasContext& CanvasContext) const
 {
-	if (CanvasContext.Canvas == nullptr)
+	UWorld* World = CanvasContext.GetWorld();
+	if (CanvasContext.Canvas == nullptr || World == nullptr)
 	{
 		return INDEX_NONE;
 	}
@@ -200,8 +206,8 @@ int32 FGameplayDebuggerCategory_EQS::DrawLookedAtItem(const EQSDebug::FQueryData
 	if (BestItemIndex != INDEX_NONE)
 	{
 		const EQSDebug::FDebugHelper& DebugHelper = QueryData.RenderDebugHelpers[BestItemIndex];
-		DrawDebugSphere(OwnerPC->GetWorld(), DebugHelper.Location, DebugHelper.Radius, 8, FColor::Red);
-		DrawDebugCone(OwnerPC->GetWorld(), DebugHelper.Location, FVector(0, 0, 1), 100.0f, 0.1f, 0.1f, 8, FColor::Red);
+		DrawDebugSphere(World, DebugHelper.Location, DebugHelper.Radius, 8, FColor::Red);
+		DrawDebugCone(World, DebugHelper.Location, FVector(0, 0, 1), 100.0f, 0.1f, 0.1f, 8, FColor::Red);
 
 		const int32 FailedTestIndex = DebugHelper.FailedTestIndex;
 		if (FailedTestIndex != INDEX_NONE)
