GameplayDebugger CIS fixes

diff --git a/Engine/Source/Developer/GameplayDebugger/Private/GameplayDebuggerPlayerManager.cpp b/Engine/Source/Developer/GameplayDebugger/Private/GameplayDebuggerPlayerManager.cpp
--- a/Engine/Source/Developer/GameplayDebugger/Private/GameplayDebuggerPlayerManager.cpp
+++ b/Engine/Source/Developer/GameplayDebugger/Private/GameplayDebuggerPlayerManager.cpp
@@ -251,7 +251,6 @@ const FGameplayDebuggerPlayerData* AGameplayDebuggerPlayerManager::GetPlayerData
 	return nullptr;
 }
 
-#if WITH_EDITOR
 // FTickableGameObject begin
 void AGameplayDebuggerPlayerManager::Tick(float DeltaTime)
 {
@@ -266,7 +265,11 @@ void AGameplayDebuggerPlayerManager::Tick(float DeltaTime)
 
 ETickableTickType AGameplayDebuggerPlayerManager::GetTickableTickType() const
 {
-	return IsTickable() ? ETickableTickType::Conditional : ETickableTickType::Never;
+	return
+#if WITH_EDITOR
+		IsTickable() ? ETickableTickType::Conditional : 
+#endif // WITH_EDITOR
+		ETickableTickType::Never;
 }
 
 TStatId AGameplayDebuggerPlayerManager::GetStatId() const
@@ -276,8 +279,12 @@ TStatId AGameplayDebuggerPlayerManager::GetStatId() const
 
 bool AGameplayDebuggerPlayerManager::IsTickable() const
 {
+#if WITH_EDITOR
 	UWorld* World = GetWorld();
 	return (World != nullptr) && (World->IsEditorWorld() == true) && (World->IsGameWorld() == false);
+#else
+	return false;
+#endif // WITH_EDITOR
 }
 // FTickableGameObject end
-#endif // WITH_EDITOR
+
diff --git a/Engine/Source/Developer/GameplayDebugger/Public/GameplayDebuggerPlayerManager.h b/Engine/Source/Developer/GameplayDebugger/Public/GameplayDebuggerPlayerManager.h
--- a/Engine/Source/Developer/GameplayDebugger/Public/GameplayDebuggerPlayerManager.h
+++ b/Engine/Source/Developer/GameplayDebugger/Public/GameplayDebuggerPlayerManager.h
@@ -35,7 +35,6 @@ class GAMEPLAYDEBUGGER_API AGameplayDebuggerPlayerManager : public AActor, publi
 
 	virtual void TickActor(float DeltaTime, enum ELevelTick TickType, FActorTickFunction& ThisTickFunction) override;
 
-#if WITH_EDITOR
 	// we're ticking only the manager only when in editor
 	// FTickableGameObject begin
 	virtual UWorld* GetTickableGameObjectWorld() const override { return GetWorld(); }
@@ -45,7 +44,6 @@ class GAMEPLAYDEBUGGER_API AGameplayDebuggerPlayerManager : public AActor, publi
 	virtual bool IsTickableInEditor() const override { return true; }
 	virtual bool IsTickable() const override;
 	// FTickableGameObject end
-#endif // WITH_EDITOR
 
 	virtual void PostInitProperties() override;
 
