Fixed a gameplay debugger tool crash on editor exit

#rnx

diff --git a/Engine/Source/Developer/GameplayDebugger/Private/Editor/GameplayDebuggerEdMode.cpp b/Engine/Source/Developer/GameplayDebugger/Private/Editor/GameplayDebuggerEdMode.cpp
--- a/Engine/Source/Developer/GameplayDebugger/Private/Editor/GameplayDebuggerEdMode.cpp
+++ b/Engine/Source/Developer/GameplayDebugger/Private/Editor/GameplayDebuggerEdMode.cpp
@@ -3,6 +3,7 @@
 #include "Editor/GameplayDebuggerEdMode.h"
 
 #if WITH_EDITOR
+#include "CoreGlobals.h"
 #include "Editor/GameplayDebuggerToolkit.h"
 #include "Toolkits/ToolkitManager.h"
 #include "GameplayDebuggerPlayerManager.h"
@@ -114,8 +115,8 @@ void FGameplayDebuggerEdMode::SafeCloseMode()
 	//
 	// DeactivateMode tries to bring up default mode, but toolkit is already destroyed by that time
 	// and editor crashes on check in GLevelEditorModeTools().GetToolkitHost() inside default mode's code
-
-	if (GLevelEditorModeTools().HasToolkitHost())
+	// note: no need to do it if IsEngineExitRequested since the GLevelEditorModeTools is no longer available 
+	if (IsEngineExitRequested() == false && GLevelEditorModeTools().HasToolkitHost())
 	{
 		GLevelEditorModeTools().DeactivateMode(FGameplayDebuggerEdMode::EM_GameplayDebugger);
 	}
