Fixed GameplayDebuggerCategoryReplicator unintentionally holding a reference to AActor being debugged which resulted in all sorts of undesired behaviors in editor mode (including fatal-level logs)..

#rnx

diff --git a/Engine/Source/Developer/GameplayDebugger/Private/GameplayDebuggerCategoryReplicator.cpp b/Engine/Source/Developer/GameplayDebugger/Private/GameplayDebuggerCategoryReplicator.cpp
--- a/Engine/Source/Developer/GameplayDebugger/Private/GameplayDebuggerCategoryReplicator.cpp
+++ b/Engine/Source/Developer/GameplayDebugger/Private/GameplayDebuggerCategoryReplicator.cpp
@@ -703,7 +703,7 @@ void AGameplayDebuggerCategoryReplicator::CollectCategoryData(bool bForce)
 			CategoryOb.ReplicatedLines.Reset();
 			CategoryOb.ReplicatedShapes.Reset();
 
-			CategoryOb.CollectData(OwnerPC, DebugActor.Actor);
+			CategoryOb.CollectData(OwnerPC, DebugActor.Actor.Get());
 			CategoryOb.LastCollectDataTime = GameTime;
 
 			// update dirty data packs
diff --git a/Engine/Source/Developer/GameplayDebugger/Public/GameplayDebuggerCategoryReplicator.h b/Engine/Source/Developer/GameplayDebugger/Public/GameplayDebuggerCategoryReplicator.h
--- a/Engine/Source/Developer/GameplayDebugger/Public/GameplayDebuggerCategoryReplicator.h
+++ b/Engine/Source/Developer/GameplayDebugger/Public/GameplayDebuggerCategoryReplicator.h
@@ -51,8 +51,7 @@ struct FGameplayDebuggerDebugActor
 {
 	GENERATED_USTRUCT_BODY()
 
-	UPROPERTY()
-	AActor* Actor;
+	TWeakObjectPtr<AActor> Actor;
 
 	UPROPERTY()
 	FName ActorName;
@@ -111,7 +110,7 @@ public:
 	void CollectCategoryData(bool bForce = false);
 
 	/** get current debug actor */
-	AActor* GetDebugActor() const { return IsValid(DebugActor.Actor) ? DebugActor.Actor : nullptr; }
+	AActor* GetDebugActor() const { return DebugActor.Actor.Get(); }
 	
 	/** get name of debug actor */
 	FName GetDebugActorName() const { return DebugActor.ActorName; }
