Made GameplayDebugger's availability in editor mode optional and cconfigured by per-user settings

I've also tabbified *.build.cs since I was modifying it anyway.

#rnx

diff --git a/Engine/Source/Developer/GameplayDebugger/GameplayDebugger.Build.cs b/Engine/Source/Developer/GameplayDebugger/GameplayDebugger.Build.cs
--- a/Engine/Source/Developer/GameplayDebugger/GameplayDebugger.Build.cs
+++ b/Engine/Source/Developer/GameplayDebugger/GameplayDebugger.Build.cs
@@ -2,47 +2,48 @@
 
 namespace UnrealBuildTool.Rules
 {
-    public class GameplayDebugger : ModuleRules
-    {
-        public GameplayDebugger(ReadOnlyTargetRules Target) : base(Target)
-        {
-            PublicDependencyModuleNames.AddRange(
-                new string[] {
-                    "Core",
-                    "CoreUObject",
-                    "Engine",
-                });
+	public class GameplayDebugger : ModuleRules
+	{
+		public GameplayDebugger(ReadOnlyTargetRules Target) : base(Target)
+		{
+			PublicDependencyModuleNames.AddRange(
+				new string[] {
+					"Core",
+					"CoreUObject",
+					"Engine",
+				});
 
-            PrivateDependencyModuleNames.AddRange(
-                new string[] {
-                    "RenderCore",
-                    "InputCore",
-                    "SlateCore",
-                    "Slate",
-                });
+			PrivateDependencyModuleNames.AddRange(
+				new string[] {
+					"RenderCore",
+					"InputCore",
+					"SlateCore",
+					"Slate",
+					"DeveloperSettings",
+				});
 
-            PrivateIncludePaths.AddRange(
-                new string[] {
-                    "Developer/GameplayDebugger/Private",
-                    "Developer/Settings/Public",
-                });
+			PrivateIncludePaths.AddRange(
+				new string[] {
+					"Developer/GameplayDebugger/Private",
+					"Developer/Settings/Public",
+				});
 
-            if (Target.bBuildEditor)
+			if (Target.bBuildEditor)
 			{
-                PrivateDependencyModuleNames.AddRange(
-                    new string[] {
-                        "EditorStyle",
+				PrivateDependencyModuleNames.AddRange(
+					new string[] {
+						"EditorStyle",
 						"EditorFramework",
 						"UnrealEd",
-                        "LevelEditor",
-                        "PropertyEditor",
-                    });
+						"LevelEditor",
+						"PropertyEditor",
+					});
 			}
 
 			if (Target.bBuildDeveloperTools || (Target.Configuration != UnrealTargetConfiguration.Shipping && Target.Configuration != UnrealTargetConfiguration.Test))
 			{
 				PrecompileForTargets = PrecompileTargetsType.Any;
 			}
-        }
-    }
+		}
+	}
 }
diff --git a/Engine/Source/Developer/GameplayDebugger/Private/GameplayDebuggerPlayerManager.cpp b/Engine/Source/Developer/GameplayDebugger/Private/GameplayDebuggerPlayerManager.cpp
--- a/Engine/Source/Developer/GameplayDebugger/Private/GameplayDebuggerPlayerManager.cpp
+++ b/Engine/Source/Developer/GameplayDebugger/Private/GameplayDebuggerPlayerManager.cpp
@@ -7,6 +7,7 @@
 #include "GameplayDebuggerLocalController.h"
 #include "Engine/DebugCameraController.h"
 #include "Engine/InputDelegateBinding.h"
+#include "GameplayDebuggerConfig.h"
 
 
 AGameplayDebuggerPlayerManager::AGameplayDebuggerPlayerManager(const FObjectInitializer& ObjectInitializer) : Super(ObjectInitializer)
@@ -97,7 +98,7 @@ void AGameplayDebuggerPlayerManager::Init()
 {
 #if WITH_EDITOR
 	UWorld* World = GetWorld();
-	if (World != nullptr && World->WorldType == EWorldType::Editor)
+	if (World != nullptr && World->WorldType == EWorldType::Editor && (GetDefault<UGameplayDebuggerUserSettings>()->bEnableGameplayDebuggerInEditor))
 	{
 		bHasAuthority = true;
 		bIsLocal = true;
diff --git a/Engine/Source/Developer/GameplayDebugger/Public/GameplayDebuggerConfig.h b/Engine/Source/Developer/GameplayDebugger/Public/GameplayDebuggerConfig.h
--- a/Engine/Source/Developer/GameplayDebugger/Public/GameplayDebuggerConfig.h
+++ b/Engine/Source/Developer/GameplayDebugger/Public/GameplayDebuggerConfig.h
@@ -6,6 +6,7 @@
 #include "UObject/ObjectMacros.h"
 #include "UObject/Object.h"
 #include "InputCoreTypes.h"
+#include "Engine/DeveloperSettings.h"
 #include "GameplayDebuggerConfig.generated.h"
 
 struct FGameplayDebuggerInputModifier;
@@ -205,3 +206,17 @@ private:
 	TMultiMap<FName, FName> KnownCategoryInputNames;
 	TMultiMap<FName, FName> KnownExtensionInputNames;
 };
+
+
+UCLASS(config = EditorPerProjectUserSettings, meta = (DisplayName = "Gameplay Debugger"))
+class GAMEPLAYDEBUGGER_API UGameplayDebuggerUserSettings : public UDeveloperSettings
+{
+	GENERATED_BODY()
+public:
+	virtual FName GetCategoryName() const override { return TEXT("Advanced"); }
+
+	/** Controls whether GamplayDebugger will be available in pure editor mode. 
+	 *  @Note that you need to reload the map for the changes to this property to take effect */
+	UPROPERTY(config, EditAnywhere, Category = GameplayDebugger)
+	uint32 bEnableGameplayDebuggerInEditor : 1;
+};
