A file missing from CL#14810194

diff --git a/Engine/Source/Developer/GameplayDebugger/Private/GameplayDebuggerCategory.cpp b/Engine/Source/Developer/GameplayDebugger/Private/GameplayDebuggerCategory.cpp
--- a/Engine/Source/Developer/GameplayDebugger/Private/GameplayDebuggerCategory.cpp
+++ b/Engine/Source/Developer/GameplayDebugger/Private/GameplayDebuggerCategory.cpp
@@ -64,7 +64,7 @@ void FGameplayDebuggerCategory::AddShape(const FGameplayDebuggerShape& Shape)
 void FGameplayDebuggerCategory::DrawCategory(APlayerController* OwnerPC, FGameplayDebuggerCanvasContext& CanvasContext)
 {
 	// note that we prefer OwnerPC's world here since if it's given then that's the end user of the data
-	UWorld* World = OwnerPC ? OwnerPC->GetWorld() : CanvasContext.World.Get();
+	UWorld* World = CanvasContext.GetWorld();
 
 	if (!ensure(World))
 	{
