[GameplayDebugger] Extracted picking view point calculation to GetViewPoint exposed through GameplayDebuggerPlayerManager

diff --git a/Engine/Source/Developer/GameplayDebugger/Private/GameplayDebuggerLocalController.cpp b/Engine/Source/Developer/GameplayDebugger/Private/GameplayDebuggerLocalController.cpp
--- a/Engine/Source/Developer/GameplayDebugger/Private/GameplayDebuggerLocalController.cpp
+++ b/Engine/Source/Developer/GameplayDebugger/Private/GameplayDebuggerLocalController.cpp
@@ -21,8 +21,6 @@
 #include "CanvasItem.h"
 #include "Engine/Canvas.h"
 #include "Engine/DebugCameraController.h"
-#include "UnrealEngine.h"
-#include "Engine/LocalPlayer.h"
 #include "GameFramework/PlayerInput.h"
 #include "EngineUtils.h"
 #include "HAL/IConsoleManager.h"
@@ -641,26 +639,9 @@ void UGameplayDebuggerLocalController::OnSelectActorTick()
 	APlayerController* OwnerPC = CachedReplicator ? CachedReplicator->GetReplicationOwner() : nullptr;
 	if (OwnerPC)
 	{
-		FVector CameraLocation;
-		FRotator CameraRotation;
-		if (OwnerPC->Player)
-		{
-			// normal game
-			OwnerPC->GetPlayerViewPoint(CameraLocation, CameraRotation);
-		}
-		else
-		{
-			// spectator mode
-			for (FLocalPlayerIterator It(GEngine, OwnerPC->GetWorld()); It; ++It)
-			{
-				ADebugCameraController* SpectatorPC = Cast<ADebugCameraController>(It->PlayerController);
-				if (SpectatorPC)
-				{
-					SpectatorPC->GetPlayerViewPoint(CameraLocation, CameraRotation);
-					break;
-				}
-			}
-		}
+		FVector ViewLocation = FVector::ZeroVector;
+		FVector ViewDirection = FVector::ForwardVector;
+		AGameplayDebuggerPlayerManager::GetViewPoint(*OwnerPC, ViewLocation, ViewDirection);
 
 		// TODO: move to module's settings
 		const float MaxScanDistance = 25000.0f;
@@ -668,19 +649,18 @@ void UGameplayDebuggerLocalController::OnSelectActorTick()
 
 		AActor* BestCandidate = nullptr;
 		float BestScore = MinViewDirDot;
-		
-		const FVector ViewDir = CameraRotation.Vector();
-		for (APawn* TestPawn  : TActorRange<APawn>(OwnerPC->GetWorld()))
+
+		for (APawn* TestPawn : TActorRange<APawn>(OwnerPC->GetWorld()))
 		{
 			if (!TestPawn->IsHidden() && TestPawn->GetActorEnableCollision() &&
 				!TestPawn->IsA(ASpectatorPawn::StaticClass()) &&
 				TestPawn != OwnerPC->GetPawn())
 			{
-				FVector DirToPawn = (TestPawn->GetActorLocation() - CameraLocation);
+				FVector DirToPawn = (TestPawn->GetActorLocation() - ViewLocation);
 				float DistToPawn = DirToPawn.Size();
 				if (FMath::IsNearlyZero(DistToPawn))
 				{
-					DirToPawn = ViewDir;
+					DirToPawn = ViewDirection;
 					DistToPawn = 1.0f;
 				}
 				else
@@ -688,7 +668,7 @@ void UGameplayDebuggerLocalController::OnSelectActorTick()
 					DirToPawn /= DistToPawn;
 				}
 
-				const float ViewDot = FVector::DotProduct(ViewDir, DirToPawn);
+				const float ViewDot = FVector::DotProduct(ViewDirection, DirToPawn);
 				if (DistToPawn < MaxScanDistance && ViewDot > BestScore)
 				{
 					BestScore = ViewDot;
diff --git a/Engine/Source/Developer/GameplayDebugger/Private/GameplayDebuggerPlayerManager.cpp b/Engine/Source/Developer/GameplayDebugger/Private/GameplayDebuggerPlayerManager.cpp
--- a/Engine/Source/Developer/GameplayDebugger/Private/GameplayDebuggerPlayerManager.cpp
+++ b/Engine/Source/Developer/GameplayDebugger/Private/GameplayDebuggerPlayerManager.cpp
@@ -8,6 +8,8 @@
 #include "Engine/DebugCameraController.h"
 #include "Engine/InputDelegateBinding.h"
 #include "GameplayDebuggerConfig.h"
+#include "UnrealEngine.h"
+#include "Engine/LocalPlayer.h"
 
 
 AGameplayDebuggerPlayerManager::AGameplayDebuggerPlayerManager(const FObjectInitializer& ObjectInitializer) : Super(ObjectInitializer)
@@ -252,6 +254,34 @@ const FGameplayDebuggerPlayerData* AGameplayDebuggerPlayerManager::GetPlayerData
 	return nullptr;
 }
 
+void AGameplayDebuggerPlayerManager::GetViewPoint(const APlayerController& OwnerPC, FVector& OutViewLocation, FVector& OutViewDirection)
+{
+	UWorld* World = OwnerPC.GetWorld();
+	FVector CameraLocation;
+	FRotator CameraRotation;
+	if (OwnerPC.Player)
+	{
+		// normal game
+		OwnerPC.GetPlayerViewPoint(CameraLocation, CameraRotation);
+	}
+	else
+	{
+		// spectator mode
+		for (FLocalPlayerIterator It(GEngine, World); It; ++It)
+		{
+			ADebugCameraController* SpectatorPC = Cast<ADebugCameraController>(It->PlayerController);
+			if (SpectatorPC)
+			{
+				SpectatorPC->GetPlayerViewPoint(CameraLocation, CameraRotation);
+				break;
+			}
+		}
+	}
+
+	OutViewLocation = CameraLocation;
+	OutViewDirection = CameraRotation.Vector();
+}
+
 // FTickableGameObject begin
 void AGameplayDebuggerPlayerManager::Tick(float DeltaTime)
 {
diff --git a/Engine/Source/Developer/GameplayDebugger/Public/GameplayDebuggerPlayerManager.h b/Engine/Source/Developer/GameplayDebugger/Public/GameplayDebuggerPlayerManager.h
--- a/Engine/Source/Developer/GameplayDebugger/Public/GameplayDebuggerPlayerManager.h
+++ b/Engine/Source/Developer/GameplayDebugger/Public/GameplayDebuggerPlayerManager.h
@@ -69,6 +69,9 @@ public:
 
 	static AGameplayDebuggerPlayerManager& GetCurrent(UWorld* World);
 
+	/** extracts view location and direction from a player controller that can be used for picking */
+	static void GetViewPoint(const APlayerController& OwnerPC, FVector& OutViewLocation, FVector& OutViewDirection);
+
 protected:
 
 	UPROPERTY()
