[FDebugDrawDelegateHelper] Allow calls to method InitDelegateHelper from `Parallel For` called from game thread (i.e. IsInParallelGameThread() || IsInGameThread() )
Removed previous hack added in UEQSRenderingComponent  and fixed some typos in method names.

diff --git a/Engine/Source/Developer/GameplayDebugger/Private/GameplayDebuggerRenderingComponent.cpp b/Engine/Source/Developer/GameplayDebugger/Private/GameplayDebuggerRenderingComponent.cpp
--- a/Engine/Source/Developer/GameplayDebugger/Private/GameplayDebuggerRenderingComponent.cpp
+++ b/Engine/Source/Developer/GameplayDebugger/Private/GameplayDebuggerRenderingComponent.cpp
@@ -77,27 +77,27 @@ protected:
 	TArray<FDebugRenderSceneProxy*> ChildProxies;
 };
 
-void FGameplayDebuggerDebugDrawDelegateHelper::RegisterDebugDrawDelgate()
+void FGameplayDebuggerDebugDrawDelegateHelper::RegisterDebugDrawDelegate()
 {
-	ensureMsgf(State != RegisteredState, TEXT("RegisterDebugDrawDelgate is already Registered!"));
+	ensureMsgf(State != RegisteredState, TEXT("RegisterDebugDrawDelegate is already Registered!"));
 	if (State == InitializedState)
 	{
 		for (int32 Idx = 0; Idx < DebugDrawDelegateHelpers.Num(); Idx++)
 		{
-			DebugDrawDelegateHelpers[Idx]->RegisterDebugDrawDelgate();
+			DebugDrawDelegateHelpers[Idx]->RegisterDebugDrawDelegate();
 		}
 		State = RegisteredState;
 	}
 }
 
-void FGameplayDebuggerDebugDrawDelegateHelper::UnregisterDebugDrawDelgate()
+void FGameplayDebuggerDebugDrawDelegateHelper::UnregisterDebugDrawDelegate()
 {
-	ensureMsgf(State != InitializedState, TEXT("UnegisterDebugDrawDelgate is in an invalid State: %i !"), State);
+	ensureMsgf(State != InitializedState, TEXT("UnregisterDebugDrawDelegate is in an invalid State: %i !"), State);
 	if (State == RegisteredState)
 	{
 		for (int32 Idx = 0; Idx < DebugDrawDelegateHelpers.Num(); Idx++)
 		{
-			DebugDrawDelegateHelpers[Idx]->UnregisterDebugDrawDelgate();
+			DebugDrawDelegateHelpers[Idx]->UnregisterDebugDrawDelegate();
 		}
 		State = InitializedState;
 	}
@@ -164,7 +164,7 @@ FPrimitiveSceneProxy* UGameplayDebuggerRenderingComponent::CreateSceneProxy()
 	if (CompositeProxy)
 	{
 		GameplayDebuggerDebugDrawDelegateHelper.InitDelegateHelper(CompositeProxy);
-		GameplayDebuggerDebugDrawDelegateHelper.ReregisterDebugDrawDelgate();
+		GameplayDebuggerDebugDrawDelegateHelper.ReregisterDebugDrawDelegate();
 	}
 	return CompositeProxy;
 }
@@ -178,12 +178,12 @@ void UGameplayDebuggerRenderingComponent::CreateRenderState_Concurrent(FRegister
 {
 	Super::CreateRenderState_Concurrent(Context);
 
-	GameplayDebuggerDebugDrawDelegateHelper.RegisterDebugDrawDelgate();
+	GameplayDebuggerDebugDrawDelegateHelper.RegisterDebugDrawDelegate();
 }
 
 void UGameplayDebuggerRenderingComponent::DestroyRenderState_Concurrent()
 {
-	GameplayDebuggerDebugDrawDelegateHelper.UnregisterDebugDrawDelgate();
+	GameplayDebuggerDebugDrawDelegateHelper.UnregisterDebugDrawDelegate();
 
 	Super::DestroyRenderState_Concurrent();
 }
diff --git a/Engine/Source/Developer/GameplayDebugger/Public/GameplayDebuggerRenderingComponent.h b/Engine/Source/Developer/GameplayDebugger/Public/GameplayDebuggerRenderingComponent.h
--- a/Engine/Source/Developer/GameplayDebugger/Public/GameplayDebuggerRenderingComponent.h
+++ b/Engine/Source/Developer/GameplayDebugger/Public/GameplayDebuggerRenderingComponent.h
@@ -24,8 +24,8 @@ public:
 
 	void AddDelegateHelper(FDebugDrawDelegateHelper* InDebugDrawDelegateHelper);
 
-	virtual void RegisterDebugDrawDelgate() override;
-	virtual void UnregisterDebugDrawDelgate() override;
+	virtual void RegisterDebugDrawDelegate() override;
+	virtual void UnregisterDebugDrawDelegate() override;
 
 private:
 	TArray<FDebugDrawDelegateHelper*> DebugDrawDelegateHelpers;
diff --git a/Engine/Source/Developer/LogVisualizer/Private/VisualLoggerRenderingActor.cpp b/Engine/Source/Developer/LogVisualizer/Private/VisualLoggerRenderingActor.cpp
--- a/Engine/Source/Developer/LogVisualizer/Private/VisualLoggerRenderingActor.cpp
+++ b/Engine/Source/Developer/LogVisualizer/Private/VisualLoggerRenderingActor.cpp
@@ -98,7 +98,7 @@ FPrimitiveSceneProxy* UVisualLoggerRenderingComponent::CreateSceneProxy()
 	if (VLogSceneProxy)
 	{
 		DebugDrawDelegateHelper.InitDelegateHelper(VLogSceneProxy);
-		DebugDrawDelegateHelper.ReregisterDebugDrawDelgate();
+		DebugDrawDelegateHelper.ReregisterDebugDrawDelegate();
 	}
 #endif
 	return VLogSceneProxy;
@@ -119,14 +119,14 @@ void UVisualLoggerRenderingComponent::CreateRenderState_Concurrent(FRegisterComp
 	Super::CreateRenderState_Concurrent(Context);
 
 #if WITH_EDITOR
-	DebugDrawDelegateHelper.RegisterDebugDrawDelgate();
+	DebugDrawDelegateHelper.RegisterDebugDrawDelegate();
 #endif
 }
 
 void UVisualLoggerRenderingComponent::DestroyRenderState_Concurrent()
 {
 #if WITH_EDITOR
-	DebugDrawDelegateHelper.UnregisterDebugDrawDelgate();
+	DebugDrawDelegateHelper.UnregisterDebugDrawDelegate();
 #endif
 
 	Super::DestroyRenderState_Concurrent();
diff --git a/Engine/Source/Runtime/AIModule/Private/EnvironmentQuery/EQSRenderingComponent.cpp b/Engine/Source/Runtime/AIModule/Private/EnvironmentQuery/EQSRenderingComponent.cpp
--- a/Engine/Source/Runtime/AIModule/Private/EnvironmentQuery/EQSRenderingComponent.cpp
+++ b/Engine/Source/Runtime/AIModule/Private/EnvironmentQuery/EQSRenderingComponent.cpp
@@ -289,15 +289,8 @@ FPrimitiveSceneProxy* UEQSRenderingComponent::CreateSceneProxy()
 #if  USE_EQS_DEBUGGER
 	if (NewSceneProxy)
 	{
-		if (IsInGameThread())
-		{
-			EQSRenderingDebugDrawDelegateHelper.InitDelegateHelper(NewSceneProxy);
-			EQSRenderingDebugDrawDelegateHelper.ReregisterDebugDrawDelgate();
-		}
-		else
-		{
-			UE_LOG(LogEQS, Warning, TEXT("Couldn't register delegate helper on non-game thread"));
-		}
+		EQSRenderingDebugDrawDelegateHelper.InitDelegateHelper(NewSceneProxy);
+		EQSRenderingDebugDrawDelegateHelper.ReregisterDebugDrawDelegate();
 	}
 #endif
 
@@ -363,14 +356,14 @@ void UEQSRenderingComponent::CreateRenderState_Concurrent(FRegisterComponentCont
 	Super::CreateRenderState_Concurrent(Context);
 
 #if USE_EQS_DEBUGGER
-	EQSRenderingDebugDrawDelegateHelper.RegisterDebugDrawDelgate();
+	EQSRenderingDebugDrawDelegateHelper.RegisterDebugDrawDelegate();
 #endif
 }
 
 void UEQSRenderingComponent::DestroyRenderState_Concurrent()
 {
 #if USE_EQS_DEBUGGER
-	EQSRenderingDebugDrawDelegateHelper.UnregisterDebugDrawDelgate();
+	EQSRenderingDebugDrawDelegateHelper.UnregisterDebugDrawDelegate();
 #endif
 
 	Super::DestroyRenderState_Concurrent();
diff --git a/Engine/Source/Runtime/Engine/Private/DebugRenderSceneProxy.cpp b/Engine/Source/Runtime/Engine/Private/DebugRenderSceneProxy.cpp
--- a/Engine/Source/Runtime/Engine/Private/DebugRenderSceneProxy.cpp
+++ b/Engine/Source/Runtime/Engine/Private/DebugRenderSceneProxy.cpp
@@ -25,7 +25,7 @@ FDebugRenderSceneProxy::FDebugRenderSceneProxy(const UPrimitiveComponent* InComp
 {
 }
 
-void FDebugDrawDelegateHelper::RegisterDebugDrawDelgate()
+void FDebugDrawDelegateHelper::RegisterDebugDrawDelegate()
 {
 	ensureMsgf(State != RegisteredState, TEXT("RegisterDebugDrawDelegate is already Registered!"));
 	if (State == InitializedState)
@@ -36,7 +36,7 @@ void FDebugDrawDelegateHelper::RegisterDebugDrawDelgate()
 	}
 }
 
-void FDebugDrawDelegateHelper::UnregisterDebugDrawDelgate()
+void FDebugDrawDelegateHelper::UnregisterDebugDrawDelegate()
 {
 	ensureMsgf(State != InitializedState, TEXT("UnregisterDebugDrawDelegate is in an invalid State: %i !"), State);
 	if (State == RegisteredState)
@@ -47,13 +47,13 @@ void FDebugDrawDelegateHelper::UnregisterDebugDrawDelgate()
 	}
 }
 
-void  FDebugDrawDelegateHelper::ReregisterDebugDrawDelgate()
+void  FDebugDrawDelegateHelper::ReregisterDebugDrawDelegate()
 {
 	ensureMsgf(State != UndefinedState, TEXT("ReregisterDebugDrawDelgeate is in an invalid State: %i !"), State);
 	if (State == RegisteredState)
 	{
-		UnregisterDebugDrawDelgate();
-		RegisterDebugDrawDelgate();
+		UnregisterDebugDrawDelegate();
+		RegisterDebugDrawDelegate();
 	}
 }
 
diff --git a/Engine/Source/Runtime/Engine/Public/DebugRenderSceneProxy.h b/Engine/Source/Runtime/Engine/Public/DebugRenderSceneProxy.h
--- a/Engine/Source/Runtime/Engine/Public/DebugRenderSceneProxy.h
+++ b/Engine/Source/Runtime/Engine/Public/DebugRenderSceneProxy.h
@@ -35,15 +35,13 @@ public:
 
 	/** 
 	 * Draw the scene proxy as a dynamic element
-	 *
-	 * @param	PDI - draw interface to render to
-	 * @param	View - current view
 	 */
 	ENGINE_API virtual void GetDynamicMeshElements(const TArray<const FSceneView*>& Views, const FSceneViewFamily& ViewFamily, uint32 VisibilityMap, FMeshElementCollector& Collector) const override;
 
 	/**
 	 * Draws a line with an arrow at the end.
 	 *
+	 * @param PDI		Draw interface to render to
 	 * @param Start		Starting point of the line.
 	 * @param End		Ending point of the line.
 	 * @param Color		Color of the line.
@@ -265,7 +263,7 @@ protected:
 public:
 	virtual void InitDelegateHelper(const FDebugRenderSceneProxy* InSceneProxy)
 	{
-		check(IsInGameThread());
+		check(IsInParallelGameThread() || IsInGameThread());
 
 		Texts.Reset();
 		Texts.Append(InSceneProxy->Texts);
@@ -274,12 +272,19 @@ public:
 		State = (State == UndefinedState) ? InitializedState : State;
 	}
 
+	UE_DEPRECATED(5.0, "This method is deprecated. Use RegisterDebugDrawDelegate instead.")
+	ENGINE_API virtual void RegisterDebugDrawDelgate() { RegisterDebugDrawDelegate(); }
+	UE_DEPRECATED(5.0, "This method is deprecated. Use UnregisterDebugDrawDelegate instead.")
+	ENGINE_API virtual void UnregisterDebugDrawDelgate() { UnregisterDebugDrawDelegate(); }
+	UE_DEPRECATED(5.0, "This method is deprecated. Use ReregisterDebugDrawDelegate instead.")
+	ENGINE_API void ReregisterDebugDrawDelgate() { ReregisterDebugDrawDelegate(); }
+
 	/** called to set up debug drawing delegate in UDebugDrawService if you want to draw labels */
-	ENGINE_API virtual void RegisterDebugDrawDelgate();
+	ENGINE_API virtual void RegisterDebugDrawDelegate();
 	/** called to clean up debug drawing delegate in UDebugDrawService */
-	ENGINE_API virtual void UnregisterDebugDrawDelgate();
+	ENGINE_API virtual void UnregisterDebugDrawDelegate();
 
-	ENGINE_API void ReregisterDebugDrawDelgate();
+	ENGINE_API void ReregisterDebugDrawDelegate();
 
 protected:
 	ENGINE_API virtual void DrawDebugLabels(UCanvas* Canvas, APlayerController*);
diff --git a/Engine/Source/Runtime/NavigationSystem/Private/NavMesh/NavMeshRenderingComponent.cpp b/Engine/Source/Runtime/NavigationSystem/Private/NavMesh/NavMeshRenderingComponent.cpp
--- a/Engine/Source/Runtime/NavigationSystem/Private/NavMesh/NavMeshRenderingComponent.cpp
+++ b/Engine/Source/Runtime/NavigationSystem/Private/NavMesh/NavMeshRenderingComponent.cpp
@@ -1132,9 +1132,9 @@ FNavMeshSceneProxy::~FNavMeshSceneProxy()
 }
 
 #if WITH_RECAST && UE_ENABLE_DEBUG_DRAWING
-void FNavMeshDebugDrawDelegateHelper::RegisterDebugDrawDelgate()
+void FNavMeshDebugDrawDelegateHelper::RegisterDebugDrawDelegate()
 {
-	ensureMsgf(State != RegisteredState, TEXT("RegisterDebugDrawDelgate is already Registered!"));
+	ensureMsgf(State != RegisteredState, TEXT("RegisterDebugDrawDelegate is already Registered!"));
 	if (State == InitializedState)
 	{
 		DebugTextDrawingDelegate = FDebugDrawDelegate::CreateRaw(this, &FNavMeshDebugDrawDelegateHelper::DrawDebugLabels);
@@ -1143,9 +1143,9 @@ void FNavMeshDebugDrawDelegateHelper::RegisterDebugDrawDelgate()
 	}
 }
 
-void FNavMeshDebugDrawDelegateHelper::UnregisterDebugDrawDelgate()
+void FNavMeshDebugDrawDelegateHelper::UnregisterDebugDrawDelegate()
 {
-	ensureMsgf(State != InitializedState, TEXT("UnegisterDebugDrawDelgate is in an invalid State: %i !"), State);
+	ensureMsgf(State != InitializedState, TEXT("UnegisterDebugDrawDelegate is in an invalid State: %i !"), State);
 	if (State == RegisteredState)
 	{
 		check(DebugTextDrawingDelegate.IsBound());
@@ -1560,8 +1560,8 @@ FPrimitiveSceneProxy* UNavMeshRenderingComponent::CreateSceneProxy()
 
 	if (NavMeshSceneProxy)
 	{
-		NavMeshDebugDrawDelgateManager.InitDelegateHelper(NavMeshSceneProxy);
-		NavMeshDebugDrawDelgateManager.ReregisterDebugDrawDelgate();
+		NavMeshDebugDrawDelegateManager.InitDelegateHelper(NavMeshSceneProxy);
+		NavMeshDebugDrawDelegateManager.ReregisterDebugDrawDelegate();
 	}
 	return NavMeshSceneProxy;
 #else
@@ -1574,14 +1574,14 @@ void UNavMeshRenderingComponent::CreateRenderState_Concurrent(FRegisterComponent
 	Super::CreateRenderState_Concurrent(Context);
 
 #if WITH_RECAST && !UE_BUILD_SHIPPING && !UE_BUILD_TEST
-	NavMeshDebugDrawDelgateManager.RegisterDebugDrawDelgate();
+	NavMeshDebugDrawDelegateManager.RegisterDebugDrawDelegate();
 #endif
 }
 
 void UNavMeshRenderingComponent::DestroyRenderState_Concurrent()
 {
 #if WITH_RECAST && !UE_BUILD_SHIPPING && !UE_BUILD_TEST
-	NavMeshDebugDrawDelgateManager.UnregisterDebugDrawDelgate();
+	NavMeshDebugDrawDelegateManager.UnregisterDebugDrawDelegate();
 #endif
 
 	Super::DestroyRenderState_Concurrent();
diff --git a/Engine/Source/Runtime/NavigationSystem/Private/NavMesh/NavTestRenderingComponent.cpp b/Engine/Source/Runtime/NavigationSystem/Private/NavMesh/NavTestRenderingComponent.cpp
--- a/Engine/Source/Runtime/NavigationSystem/Private/NavMesh/NavTestRenderingComponent.cpp
+++ b/Engine/Source/Runtime/NavigationSystem/Private/NavMesh/NavTestRenderingComponent.cpp
@@ -331,9 +331,9 @@ void FNavTestDebugDrawDelegateHelper::InitDelegateHelper(const FNavTestSceneProx
 	bShowDiff = InSceneProxy->bShowDiff;
 }
 
-void FNavTestDebugDrawDelegateHelper::RegisterDebugDrawDelgate()
+void FNavTestDebugDrawDelegateHelper::RegisterDebugDrawDelegate()
 {
-	ensureMsgf(State != RegisteredState, TEXT("RegisterDebugDrawDelgate is already Registered!"));
+	ensureMsgf(State != RegisteredState, TEXT("RegisterDebugDrawDelegate is already Registered!"));
 	if (State == InitializedState)
 	{
 		DebugTextDrawingDelegate = FDebugDrawDelegate::CreateRaw(this, &FNavTestDebugDrawDelegateHelper::DrawDebugLabels);
@@ -342,9 +342,9 @@ void FNavTestDebugDrawDelegateHelper::RegisterDebugDrawDelgate()
 	}
 }
 
-void FNavTestDebugDrawDelegateHelper::UnregisterDebugDrawDelgate()
+void FNavTestDebugDrawDelegateHelper::UnregisterDebugDrawDelegate()
 {
-	ensureMsgf(State != InitializedState, TEXT("UnegisterDebugDrawDelgate is in an invalid State: %i !"), State);
+	ensureMsgf(State != InitializedState, TEXT("UnegisterDebugDrawDelegate is in an invalid State: %i !"), State);
 	if (State == RegisteredState)
 	{
 		check(DebugTextDrawingDelegate.IsBound());
@@ -420,7 +420,7 @@ FPrimitiveSceneProxy* UNavTestRenderingComponent::CreateSceneProxy()
 	FNavTestSceneProxy* SceneProxy2 = new FNavTestSceneProxy(this);
 #if WITH_RECAST && WITH_EDITOR
 	NavTestDebugDrawDelegateHelper.InitDelegateHelper(SceneProxy2);
-	NavTestDebugDrawDelegateHelper.ReregisterDebugDrawDelgate();
+	NavTestDebugDrawDelegateHelper.ReregisterDebugDrawDelegate();
 #endif
 	return SceneProxy2;
 }
@@ -466,14 +466,14 @@ void UNavTestRenderingComponent::CreateRenderState_Concurrent(FRegisterComponent
 	Super::CreateRenderState_Concurrent(Context);
 
 #if WITH_RECAST && WITH_EDITOR
-	NavTestDebugDrawDelegateHelper.RegisterDebugDrawDelgate();
+	NavTestDebugDrawDelegateHelper.RegisterDebugDrawDelegate();
 #endif
 }
 
 void UNavTestRenderingComponent::DestroyRenderState_Concurrent()
 {
 #if WITH_RECAST && WITH_EDITOR
-	NavTestDebugDrawDelegateHelper.UnregisterDebugDrawDelgate();
+	NavTestDebugDrawDelegateHelper.UnregisterDebugDrawDelegate();
 #endif
 
 	Super::DestroyRenderState_Concurrent();
diff --git a/Engine/Source/Runtime/NavigationSystem/Public/NavMesh/NavMeshRenderingComponent.h b/Engine/Source/Runtime/NavigationSystem/Public/NavMesh/NavMeshRenderingComponent.h
--- a/Engine/Source/Runtime/NavigationSystem/Public/NavMesh/NavMeshRenderingComponent.h
+++ b/Engine/Source/Runtime/NavigationSystem/Public/NavMesh/NavMeshRenderingComponent.h
@@ -176,8 +176,8 @@ public:
 		bNeedsNewData = InSceneProxy->ProxyData.bNeedsNewData;
 	}
 
-	NAVIGATIONSYSTEM_API virtual void RegisterDebugDrawDelgate() override;
-	NAVIGATIONSYSTEM_API virtual void UnregisterDebugDrawDelgate() override;
+	NAVIGATIONSYSTEM_API virtual void RegisterDebugDrawDelegate() override;
+	NAVIGATIONSYSTEM_API virtual void UnregisterDebugDrawDelegate() override;
 
 protected:
 	NAVIGATIONSYSTEM_API virtual void DrawDebugLabels(UCanvas* Canvas, APlayerController*) override;
@@ -230,7 +230,7 @@ protected:
 
 protected:
 #if WITH_RECAST && UE_ENABLE_DEBUG_DRAWING
-	FNavMeshDebugDrawDelegateHelper NavMeshDebugDrawDelgateManager;
+	FNavMeshDebugDrawDelegateHelper NavMeshDebugDrawDelegateManager;
 #endif
 };
 
diff --git a/Engine/Source/Runtime/NavigationSystem/Public/NavMesh/NavTestRenderingComponent.h b/Engine/Source/Runtime/NavigationSystem/Public/NavMesh/NavTestRenderingComponent.h
--- a/Engine/Source/Runtime/NavigationSystem/Public/NavMesh/NavTestRenderingComponent.h
+++ b/Engine/Source/Runtime/NavigationSystem/Public/NavMesh/NavTestRenderingComponent.h
@@ -101,8 +101,8 @@ public:
 
 	void InitDelegateHelper(const FNavTestSceneProxy* InSceneProxy);
 
-	virtual void RegisterDebugDrawDelgate() override;
-	virtual void UnregisterDebugDrawDelgate() override;
+	virtual void RegisterDebugDrawDelegate() override;
+	virtual void UnregisterDebugDrawDelegate() override;
 
 protected:
 	void DrawDebugLabels(UCanvas* Canvas, APlayerController*) override;
