GameplayDebugger: Added arrow and circle shape.

diff --git a/Engine/Source/Developer/GameplayDebugger/Private/GameplayDebuggerTypes.cpp b/Engine/Source/Developer/GameplayDebugger/Private/GameplayDebuggerTypes.cpp
--- a/Engine/Source/Developer/GameplayDebugger/Private/GameplayDebuggerTypes.cpp
+++ b/Engine/Source/Developer/GameplayDebugger/Private/GameplayDebuggerTypes.cpp
@@ -54,6 +54,19 @@ FGameplayDebuggerShape FGameplayDebuggerShape::MakeSegment(const FVector& StartL
 	return MakeSegment(StartLocation, EndLocation, 1.0f, Color, Description);
 }
 
+FGameplayDebuggerShape FGameplayDebuggerShape::MakeArrow(const FVector& StartLocation, const FVector& EndLocation, const float HeadSize, const float Thickness, const FColor& Color, const FString& Description)
+{
+	FGameplayDebuggerShape NewElement;
+	NewElement.ShapeData.Add(StartLocation);
+	NewElement.ShapeData.Add(EndLocation);
+	NewElement.ShapeData.Add(FVector(Thickness, HeadSize, 0));
+	NewElement.Color = Color;
+	NewElement.Description = Description;
+	NewElement.Type = EGameplayDebuggerShape::Arrow;
+
+	return NewElement;
+}
+
 FGameplayDebuggerShape FGameplayDebuggerShape::MakeBox(const FVector& Center, const FVector& Extent, const FColor& Color, const FString& Description)
 {
 	FGameplayDebuggerShape NewElement;
@@ -91,6 +104,19 @@ FGameplayDebuggerShape FGameplayDebuggerShape::MakeCylinder(const FVector& Cente
 	return NewElement;
 }
 
+FGameplayDebuggerShape FGameplayDebuggerShape::MakeCircle(const FVector& Center, const FVector& Up, const float Radius, const FColor& Color, const FString& Description)
+{
+	FGameplayDebuggerShape NewElement;
+	NewElement.ShapeData.Add(Center);
+	NewElement.ShapeData.Add(Up);
+	NewElement.ShapeData.Add(FVector(Radius, 0, 0));
+	NewElement.Color = Color;
+	NewElement.Description = Description;
+	NewElement.Type = EGameplayDebuggerShape::Circle;
+
+	return NewElement;
+}
+
 FGameplayDebuggerShape FGameplayDebuggerShape::MakeCapsule(const FVector& Center, const float Radius, const float HalfHeight, const FColor& Color, const FString& Description)
 {
 	FGameplayDebuggerShape NewElement;
@@ -115,13 +141,17 @@ FGameplayDebuggerShape FGameplayDebuggerShape::MakePolygon(const TArray<FVector>
 
 void FGameplayDebuggerShape::Draw(UWorld* World, FGameplayDebuggerCanvasContext& Context)
 {
+	constexpr bool bPersistent = false;
+	constexpr float LifeTime = -1.0f;
+	constexpr uint8 DepthPriority = SDPG_World;
+	
 	FVector DescLocation;
 	switch (Type)
 	{
 	case EGameplayDebuggerShape::Point:
 		if (ShapeData.Num() == 2 && ShapeData[1].X > 0)
 		{
-			DrawDebugSphere(World, ShapeData[0], ShapeData[1].X, 16, Color);
+			DrawDebugSphere(World, ShapeData[0], ShapeData[1].X, 16, Color, bPersistent, LifeTime, DepthPriority);
 			DescLocation = ShapeData[0];
 		}
 		break;
@@ -129,7 +159,15 @@ void FGameplayDebuggerShape::Draw(UWorld* World, FGameplayDebuggerCanvasContext&
 	case EGameplayDebuggerShape::Segment:
 		if (ShapeData.Num() == 3 && ShapeData[2].X > 0)
 		{
-			DrawDebugLine(World, ShapeData[0], ShapeData[1], Color, false, -1.0f, 0, ShapeData[2].X);
+			DrawDebugLine(World, ShapeData[0], ShapeData[1], Color, bPersistent, LifeTime, DepthPriority, ShapeData[2].X);
+			DescLocation = (ShapeData[0] + ShapeData[1]) * 0.5f;
+		}
+		break;
+
+	case EGameplayDebuggerShape::Arrow:
+		if (ShapeData.Num() == 3 && ShapeData[2].X > 0)
+		{
+			DrawDebugDirectionalArrow(World, ShapeData[0], ShapeData[1], ShapeData[2].Y, Color, bPersistent, LifeTime, DepthPriority, ShapeData[2].X);
 			DescLocation = (ShapeData[0] + ShapeData[1]) * 0.5f;
 		}
 		break;
@@ -146,7 +184,7 @@ void FGameplayDebuggerShape::Draw(UWorld* World, FGameplayDebuggerCanvasContext&
 		if (ShapeData.Num() == 3 && ShapeData[2].X > 0)
 		{
 			const float DefaultConeAngle = 0.25f; // ~ 15 degrees
-			DrawDebugCone(World, ShapeData[0], ShapeData[1], ShapeData[2].X, DefaultConeAngle, DefaultConeAngle, 16, Color);
+			DrawDebugCone(World, ShapeData[0], ShapeData[1], ShapeData[2].X, DefaultConeAngle, DefaultConeAngle, 16, Color, bPersistent, LifeTime, DepthPriority);
 			DescLocation = ShapeData[0];
 		}
 		break;
@@ -154,7 +192,16 @@ void FGameplayDebuggerShape::Draw(UWorld* World, FGameplayDebuggerCanvasContext&
 	case EGameplayDebuggerShape::Cylinder:
 		if (ShapeData.Num() == 2)
 		{
-			DrawDebugCylinder(World, ShapeData[0] - FVector(0, 0, ShapeData[1].Z), ShapeData[0] + FVector(0, 0, ShapeData[1].Z), ShapeData[1].X, 16, Color);
+			DrawDebugCylinder(World, ShapeData[0] - FVector(0, 0, ShapeData[1].Z), ShapeData[0] + FVector(0, 0, ShapeData[1].Z), ShapeData[1].X, 16, Color, bPersistent, LifeTime, DepthPriority);
+			DescLocation = ShapeData[0];
+		}
+		break;
+
+	case EGameplayDebuggerShape::Circle:
+		if (ShapeData.Num() == 3)
+		{
+			const FMatrix Axes = FRotationMatrix::MakeFromX(ShapeData[1]);
+			DrawDebugCircle(World, ShapeData[0], ShapeData[2].X, 32, FColor::White, bPersistent, LifeTime, DepthPriority, 1, Axes.GetUnitAxis(EAxis::Y), Axes.GetUnitAxis(EAxis::Z), false);
 			DescLocation = ShapeData[0];
 		}
 		break;
@@ -162,7 +209,7 @@ void FGameplayDebuggerShape::Draw(UWorld* World, FGameplayDebuggerCanvasContext&
 	case EGameplayDebuggerShape::Capsule:
 		if (ShapeData.Num() == 2)
 		{
-			DrawDebugCapsule(World, ShapeData[0], ShapeData[1].Z, ShapeData[1].X, FQuat::Identity, Color);
+			DrawDebugCapsule(World, ShapeData[0], ShapeData[1].Z, ShapeData[1].X, FQuat::Identity, Color, bPersistent, LifeTime, DepthPriority);
 			DescLocation = ShapeData[0];
 		}
 		break;
@@ -178,7 +225,7 @@ void FGameplayDebuggerShape::Draw(UWorld* World, FGameplayDebuggerCanvasContext&
 				MidPoint += ShapeData[Idx];
 			}
 
-			DrawDebugMesh(World, ShapeData, Indices, Color);
+			DrawDebugMesh(World, ShapeData, Indices, Color, bPersistent, LifeTime, DepthPriority);
 			DescLocation = MidPoint / ShapeData.Num();
 		}
 		break;
diff --git a/Engine/Source/Developer/GameplayDebugger/Public/GameplayDebuggerTypes.h b/Engine/Source/Developer/GameplayDebugger/Public/GameplayDebuggerTypes.h
--- a/Engine/Source/Developer/GameplayDebugger/Public/GameplayDebuggerTypes.h
+++ b/Engine/Source/Developer/GameplayDebugger/Public/GameplayDebuggerTypes.h
@@ -142,8 +142,10 @@ enum class EGameplayDebuggerShape : uint8
 	Box,
 	Cone,
 	Cylinder,
+	Circle,
 	Capsule,
 	Polygon,
+	Arrow,
 };
 
 struct GAMEPLAYDEBUGGER_API FGameplayDebuggerShape
@@ -170,9 +172,11 @@ struct GAMEPLAYDEBUGGER_API FGameplayDebuggerShape
 	static FGameplayDebuggerShape MakePoint(const FVector& Location, const float Radius, const FColor& Color, const FString& Description = FString());
 	static FGameplayDebuggerShape MakeSegment(const FVector& StartLocation, const FVector& EndLocation, const float Thickness, const FColor& Color, const FString& Description = FString());
 	static FGameplayDebuggerShape MakeSegment(const FVector& StartLocation, const FVector& EndLocation, const FColor& Color, const FString& Description = FString());
+	static FGameplayDebuggerShape MakeArrow(const FVector& StartLocation, const FVector& EndLocation, const float HeadSize, const float Thickness, const FColor& Color, const FString& Description = FString());
 	static FGameplayDebuggerShape MakeBox(const FVector& Center, const FVector& Extent, const FColor& Color, const FString& Description = FString());
 	static FGameplayDebuggerShape MakeCone(const FVector& Location, const FVector& Direction, const float Length, const FColor& Color, const FString& Description = FString());
 	static FGameplayDebuggerShape MakeCylinder(const FVector& Center, const float Radius, const float HalfHeight, const FColor& Color, const FString& Description = FString());
+	static FGameplayDebuggerShape MakeCircle(const FVector& Center, const FVector& Up, const float Radius, const FColor& Color, const FString& Description = FString());
 	static FGameplayDebuggerShape MakeCapsule(const FVector& Center, const float Radius, const float HalfHeight, const FColor& Color, const FString& Description = FString());
 	static FGameplayDebuggerShape MakePolygon(const TArray<FVector>& Verts, const FColor& Color, const FString& Description = FString());
 
