GameplayDebugger: Fixed color on circle shape

diff --git a/Engine/Source/Developer/GameplayDebugger/Private/GameplayDebuggerTypes.cpp b/Engine/Source/Developer/GameplayDebugger/Private/GameplayDebuggerTypes.cpp
--- a/Engine/Source/Developer/GameplayDebugger/Private/GameplayDebuggerTypes.cpp
+++ b/Engine/Source/Developer/GameplayDebugger/Private/GameplayDebuggerTypes.cpp
@@ -201,7 +201,7 @@ void FGameplayDebuggerShape::Draw(UWorld* World, FGameplayDebuggerCanvasContext&
 		if (ShapeData.Num() == 3)
 		{
 			const FMatrix Axes = FRotationMatrix::MakeFromX(ShapeData[1]);
-			DrawDebugCircle(World, ShapeData[0], ShapeData[2].X, 32, FColor::White, bPersistent, LifeTime, DepthPriority, 1, Axes.GetUnitAxis(EAxis::Y), Axes.GetUnitAxis(EAxis::Z), false);
+			DrawDebugCircle(World, ShapeData[0], ShapeData[2].X, 32, Color, bPersistent, LifeTime, DepthPriority, 1, Axes.GetUnitAxis(EAxis::Y), Axes.GetUnitAxis(EAxis::Z), false);
 			DescLocation = ShapeData[0];
 		}
 		break;
