[AI] Merging 17468937: GameplayDebugger: Support separate alpha on text

diff --git a/Engine/Source/Developer/GameplayDebugger/Private/GameplayDebuggerTypes.cpp b/Engine/Source/Developer/GameplayDebugger/Private/GameplayDebuggerTypes.cpp
--- a/Engine/Source/Developer/GameplayDebugger/Private/GameplayDebuggerTypes.cpp
+++ b/Engine/Source/Developer/GameplayDebugger/Private/GameplayDebuggerTypes.cpp
@@ -419,10 +419,15 @@ FGameplayDebuggerCanvasContext::FGameplayDebuggerCanvasContext(UCanvas* InCanvas
 
 void FGameplayDebuggerCanvasContext::Print(const FString& String)
 {
-	Print(FColor::White, String);
+	Print(FColor::White, 1.0f, String);
 }
 
 void FGameplayDebuggerCanvasContext::Print(const FColor& Color, const FString& String)
+{
+	Print(Color,  1.0f, String);
+}
+
+void FGameplayDebuggerCanvasContext::Print(const FColor& Color, const float Alpha, const FString& String)
 {
 	FTaggedStringParser Parser(Color);
 	Parser.ParseString(String);
@@ -448,7 +453,10 @@ void FGameplayDebuggerCanvasContext::Print(const FColor& Color, const FString& S
 			float SizeX = 0.0f, SizeY = 0.0f;
 			MeasureString(NodeData.String, SizeX, SizeY);
 
-			FCanvasTextItem TextItem(FVector2D::ZeroVector, FText::FromString(NodeData.String), Font.Get(), FLinearColor(NodeData.Color));
+			FLinearColor TextColor(NodeData.Color);
+			TextColor.A = Alpha;
+			
+			FCanvasTextItem TextItem(FVector2D::ZeroVector, FText::FromString(NodeData.String), Font.Get(), TextColor);
 			if (FontRenderInfo.bEnableShadow)
 			{
 				TextItem.EnableShadow(FColor::Black, FVector2D(1, 1));
@@ -464,17 +472,22 @@ void FGameplayDebuggerCanvasContext::Print(const FColor& Color, const FString& S
 
 void FGameplayDebuggerCanvasContext::PrintAt(float PosX, float PosY, const FString& String)
 {
-	PrintAt(PosX, PosY, FColor::White, String);
+	PrintAt(PosX, PosY, FColor::White, 1.0f, String);
 }
 
 void FGameplayDebuggerCanvasContext::PrintAt(float PosX, float PosY, const FColor& Color, const FString& String)
+{
+	PrintAt(PosX, PosY, Color, 1.0f, String);
+}
+
+void FGameplayDebuggerCanvasContext::PrintAt(float PosX, float PosY, const FColor& Color, const float Alpha, const FString& String)
 {
 	TGuardValue<float> ScopedCursorX(CursorX, PosX);
 	TGuardValue<float> ScopedCursorY(CursorY, PosY);
 	TGuardValue<float> ScopedDefaultX(DefaultX, PosX);
 	TGuardValue<float> ScopedDefaultY(DefaultY, PosY);
 
-	Print(Color, String);
+	Print(Color, Alpha, String);
 }
 
 // copied from Core/Private/Misc/VarargsHeler.h 
@@ -505,24 +518,14 @@ void FGameplayDebuggerCanvasContext::PrintAt(float PosX, float PosY, const FColo
 	PrintFunc; \
 	FMemory::SystemFree(AllocatedBuffer);
 
-void FGameplayDebuggerCanvasContext::PrintfImpl(const TCHAR* Fmt, ...)
-{
-	GROWABLE_PRINTF(Print(Buffer));
-}
-
-void FGameplayDebuggerCanvasContext::PrintfImpl(const FColor& Color, const TCHAR* Fmt, ...)
-{
-	GROWABLE_PRINTF(Print(Color, Buffer));
-}
-
-void FGameplayDebuggerCanvasContext::PrintfAtImpl(float PosX, float PosY, const TCHAR* Fmt, ...)
+void FGameplayDebuggerCanvasContext::PrintfImpl(const FColor& Color, const float Alpha, const TCHAR* Fmt, ...)
 {
-	GROWABLE_PRINTF(PrintAt(PosX, PosY, Buffer));
+	GROWABLE_PRINTF(Print(Color, Alpha, Buffer));
 }
 
-void FGameplayDebuggerCanvasContext::PrintfAtImpl(float PosX, float PosY, const FColor& Color, const TCHAR* Fmt, ...)
+void FGameplayDebuggerCanvasContext::PrintfAtImpl(float PosX, float PosY, const FColor& Color, const float Alpha, const TCHAR* Fmt, ...)
 {
-	GROWABLE_PRINTF(PrintAt(PosX, PosY, Color, Buffer));
+	GROWABLE_PRINTF(PrintAt(PosX, PosY, Color, Alpha, Buffer));
 }
 
 void FGameplayDebuggerCanvasContext::MoveToNewLine()
diff --git a/Engine/Source/Developer/GameplayDebugger/Public/GameplayDebuggerTypes.h b/Engine/Source/Developer/GameplayDebugger/Public/GameplayDebuggerTypes.h
--- a/Engine/Source/Developer/GameplayDebugger/Public/GameplayDebuggerTypes.h
+++ b/Engine/Source/Developer/GameplayDebugger/Public/GameplayDebuggerTypes.h
@@ -52,8 +52,10 @@ public:
 	// print string on canvas
 	void Print(const FString& String);
 	void Print(const FColor& Color, const FString& String);
-	void PrintAt(float PosX, float PosY, const FColor& Color, const FString& String);
+	void Print(const FColor& Color, const float Alpha, const FString& String);
 	void PrintAt(float PosX, float PosY, const FString& String);
+	void PrintAt(float PosX, float PosY, const FColor& Color, const FString& String);
+	void PrintAt(float PosX, float PosY, const FColor& Color, const float Alpha, const FString& String);
 
 	// print formatted string on canvas
 	template <typename FmtType, typename... Types>
@@ -62,38 +64,57 @@ public:
 		static_assert(TIsArrayOrRefOfType<FmtType, TCHAR>::Value, "Formatting string must be a TCHAR array.");
 		static_assert(TAnd<TIsValidVariadicFunctionArg<Types>...>::Value, "Invalid argument(s) passed to FGameplayDebuggerCanvasContext::Printf");
 
-		PrintfImpl(Fmt, Args...);
+		PrintfImpl(FColor::White, 1.0f, Fmt, Args...);
 	}
+	
 	template <typename FmtType, typename... Types>
 	void Printf(const FColor& Color, const FmtType& Fmt, Types... Args)
 	{
 		static_assert(TIsArrayOrRefOfType<FmtType, TCHAR>::Value, "Formatting string must be a TCHAR array.");
 		static_assert(TAnd<TIsValidVariadicFunctionArg<Types>...>::Value, "Invalid argument(s) passed to FGameplayDebuggerCanvasContext::Printf");
 
-		PrintfImpl(Color, Fmt, Args...);
+		PrintfImpl(Color, 1.0f, Fmt, Args...);
+	}
+	
+	template <typename FmtType, typename... Types>
+	void Printf(const FColor& Color, const float Alpha, const FmtType& Fmt, Types... Args)
+	{
+		static_assert(TIsArrayOrRefOfType<FmtType, TCHAR>::Value, "Formatting string must be a TCHAR array.");
+		static_assert(TAnd<TIsValidVariadicFunctionArg<Types>...>::Value, "Invalid argument(s) passed to FGameplayDebuggerCanvasContext::Printf");
+
+		PrintfImpl(Color, Alpha, Fmt, Args...);
 	}
+	
 	template <typename FmtType, typename... Types>
 	void PrintfAt(float PosX, float PosY, const FmtType& Fmt, Types... Args)
 	{
 		static_assert(TIsArrayOrRefOfType<FmtType, TCHAR>::Value, "Formatting string must be a TCHAR array.");
 		static_assert(TAnd<TIsValidVariadicFunctionArg<Types>...>::Value, "Invalid argument(s) passed to FGameplayDebuggerCanvasContext::PrintfAt");
 
-		PrintfAtImpl(PosX, PosY, Fmt, Args...);
+		PrintfAtImpl(PosX, PosY, FColor::White, 1.0f, Fmt, Args...);
 	}
+	
 	template <typename FmtType, typename... Types>
 	void PrintfAt(float PosX, float PosY, const FColor& Color, const FmtType& Fmt, Types... Args)
 	{
 		static_assert(TIsArrayOrRefOfType<FmtType, TCHAR>::Value, "Formatting string must be a TCHAR array.");
 		static_assert(TAnd<TIsValidVariadicFunctionArg<Types>...>::Value, "Invalid argument(s) passed to FGameplayDebuggerCanvasContext::PrintfAt");
 
-		PrintfAtImpl(PosX, PosY, Color, Fmt, Args...);
+		PrintfAtImpl(PosX, PosY, Color, 1.0f, Fmt, Args...);
+	}
+	
+	template <typename FmtType, typename... Types>
+	void PrintfAt(float PosX, float PosY, const FColor& Color, const float Alpha, const FmtType& Fmt, Types... Args)
+	{
+		static_assert(TIsArrayOrRefOfType<FmtType, TCHAR>::Value, "Formatting string must be a TCHAR array.");
+		static_assert(TAnd<TIsValidVariadicFunctionArg<Types>...>::Value, "Invalid argument(s) passed to FGameplayDebuggerCanvasContext::PrintfAt");
+
+		PrintfAtImpl(PosX, PosY, Color, Alpha, Fmt, Args...);
 	}
 
 private:
-	void VARARGS PrintfImpl(const TCHAR* Args, ...);
-	void VARARGS PrintfImpl(const FColor& Color, const TCHAR* Args, ...);
-	void VARARGS PrintfAtImpl(float PosX, float PosY, const TCHAR* Args, ...);
-	void VARARGS PrintfAtImpl(float PosX, float PosY, const FColor& Color, const TCHAR* Args, ...);
+	void VARARGS PrintfImpl(const FColor& Color, const float Alpha, const TCHAR* Args, ...);
+	void VARARGS PrintfAtImpl(float PosX, float PosY, const FColor& Color, const float Alpha, const TCHAR* Args, ...);
 
 public:
 	// moves cursor to new line
