[GameplayDebugger] added MakeCircle overload that supports Thickness

diff --git a/Engine/Source/Developer/GameplayDebugger/Private/GameplayDebuggerTypes.cpp b/Engine/Source/Developer/GameplayDebugger/Private/GameplayDebuggerTypes.cpp
--- a/Engine/Source/Developer/GameplayDebugger/Private/GameplayDebuggerTypes.cpp
+++ b/Engine/Source/Developer/GameplayDebugger/Private/GameplayDebuggerTypes.cpp
@@ -105,11 +105,16 @@ FGameplayDebuggerShape FGameplayDebuggerShape::MakeCylinder(const FVector& Cente
 }
 
 FGameplayDebuggerShape FGameplayDebuggerShape::MakeCircle(const FVector& Center, const FVector& Up, const float Radius, const FColor& Color, const FString& Description)
+{
+	return MakeCircle(Center, Up, Radius, 1.f, Color, Description);
+}
+
+FGameplayDebuggerShape FGameplayDebuggerShape::MakeCircle(const FVector& Center, const FVector& Up, const float Radius, const float Thickness, const FColor& Color, const FString& Description)
 {
 	FGameplayDebuggerShape NewElement;
 	NewElement.ShapeData.Add(Center);
 	NewElement.ShapeData.Add(Up);
-	NewElement.ShapeData.Add(FVector(Radius, 0, 0));
+	NewElement.ShapeData.Add(FVector(Radius, Thickness, 0));
 	NewElement.Color = Color;
 	NewElement.Description = Description;
 	NewElement.Type = EGameplayDebuggerShape::Circle;
@@ -201,7 +206,7 @@ void FGameplayDebuggerShape::Draw(UWorld* World, FGameplayDebuggerCanvasContext&
 		if (ShapeData.Num() == 3)
 		{
 			const FMatrix Axes = FRotationMatrix::MakeFromX(ShapeData[1]);
-			DrawDebugCircle(World, ShapeData[0], ShapeData[2].X, 32, Color, bPersistent, LifeTime, DepthPriority, 1, Axes.GetUnitAxis(EAxis::Y), Axes.GetUnitAxis(EAxis::Z), false);
+			DrawDebugCircle(World, ShapeData[0], ShapeData[2].X, 32, Color, bPersistent, LifeTime, DepthPriority, ShapeData[2].Y, Axes.GetUnitAxis(EAxis::Y), Axes.GetUnitAxis(EAxis::Z), false);
 			DescLocation = ShapeData[0];
 		}
 		break;
diff --git a/Engine/Source/Developer/GameplayDebugger/Public/GameplayDebuggerTypes.h b/Engine/Source/Developer/GameplayDebugger/Public/GameplayDebuggerTypes.h
--- a/Engine/Source/Developer/GameplayDebugger/Public/GameplayDebuggerTypes.h
+++ b/Engine/Source/Developer/GameplayDebugger/Public/GameplayDebuggerTypes.h
@@ -198,6 +198,7 @@ struct GAMEPLAYDEBUGGER_API FGameplayDebuggerShape
 	static FGameplayDebuggerShape MakeCone(const FVector& Location, const FVector& Direction, const float Length, const FColor& Color, const FString& Description = FString());
 	static FGameplayDebuggerShape MakeCylinder(const FVector& Center, const float Radius, const float HalfHeight, const FColor& Color, const FString& Description = FString());
 	static FGameplayDebuggerShape MakeCircle(const FVector& Center, const FVector& Up, const float Radius, const FColor& Color, const FString& Description = FString());
+	static FGameplayDebuggerShape MakeCircle(const FVector& Center, const FVector& Up, const float Radius, const float Thickness, const FColor& Color, const FString& Description = FString());
 	static FGameplayDebuggerShape MakeCapsule(const FVector& Center, const float Radius, const float HalfHeight, const FColor& Color, const FString& Description = FString());
 	static FGameplayDebuggerShape MakePolygon(const TArray<FVector>& Verts, const FColor& Color, const FString& Description = FString());
 
