Fix for shootergame rep graph crash after seamless travel

diff --git a/Engine/Source/Developer/GameplayDebugger/Private/GameplayDebuggerCategoryReplicator.cpp b/Engine/Source/Developer/GameplayDebugger/Private/GameplayDebuggerCategoryReplicator.cpp
--- a/Engine/Source/Developer/GameplayDebugger/Private/GameplayDebuggerCategoryReplicator.cpp
+++ b/Engine/Source/Developer/GameplayDebugger/Private/GameplayDebuggerCategoryReplicator.cpp
@@ -472,6 +472,8 @@ void AGameplayDebuggerCategoryReplicator::EndPlay(const EEndPlayReason::Type End
 {
 	Super::EndPlay(EndPlayReason);
 
+	SetReplicatorOwner(nullptr);
+
 	// Disable extensions to clear UI state
 	NotifyCategoriesToolState(false);
 	NotifyExtensionsToolState(false);
