[GameplayDebugger] use scene view to replicated view matrix only when associated actor is not available.
#rnx

diff --git a/Engine/Source/Developer/GameplayDebugger/Private/GameplayDebuggerCategoryReplicator.cpp b/Engine/Source/Developer/GameplayDebugger/Private/GameplayDebuggerCategoryReplicator.cpp
--- a/Engine/Source/Developer/GameplayDebugger/Private/GameplayDebuggerCategoryReplicator.cpp
+++ b/Engine/Source/Developer/GameplayDebugger/Private/GameplayDebuggerCategoryReplicator.cpp
@@ -562,6 +562,16 @@ void AGameplayDebuggerCategoryReplicator::ServerSetViewPoint_Implementation(cons
 	SetViewPoint(InViewLocation, InViewDirection);
 }
 
+bool AGameplayDebuggerCategoryReplicator::ServerResetViewPoint_Validate()
+{
+	return true;
+}
+
+void AGameplayDebuggerCategoryReplicator::ServerResetViewPoint_Implementation()
+{
+	ResetViewPoint();
+}
+
 bool AGameplayDebuggerCategoryReplicator::ServerSetCategoryEnabled_Validate(int32 CategoryId, bool bEnable)
 {
 	return true;
@@ -752,7 +762,7 @@ void AGameplayDebuggerCategoryReplicator::CollectCategoryData(bool bForce)
 
 bool AGameplayDebuggerCategoryReplicator::GetViewPoint(FVector& OutViewLocation, FVector& OutViewDirection) const
 {
-	if (ViewLocation.IsSet() && ViewDirection.IsSet())
+	if (IsViewPointSet())
 	{
 		OutViewLocation = ViewLocation.GetValue();
 		OutViewDirection = ViewDirection.GetValue();
@@ -836,14 +846,23 @@ void AGameplayDebuggerCategoryReplicator::SetDebugActor(AActor* Actor, bool bSel
 
 void AGameplayDebuggerCategoryReplicator::SetViewPoint(const FVector& InViewLocation, const FVector& InViewDirection)
 {
-	if (bHasAuthority)
+	ViewLocation = InViewLocation;
+	ViewDirection = InViewDirection;
+
+	if (!bHasAuthority)
 	{
-		ViewLocation = InViewLocation;
-		ViewDirection = InViewDirection;
+		ServerSetViewPoint(InViewLocation, InViewDirection);
 	}
-	else
+}
+
+void AGameplayDebuggerCategoryReplicator::ResetViewPoint()
+{
+	ViewLocation.Reset();
+	ViewDirection.Reset();
+
+	if (!bHasAuthority)
 	{
-		ServerSetViewPoint(InViewLocation, InViewDirection);
+		ServerResetViewPoint();
 	}
 }
 
diff --git a/Engine/Source/Developer/GameplayDebugger/Private/GameplayDebuggerLocalController.cpp b/Engine/Source/Developer/GameplayDebugger/Private/GameplayDebuggerLocalController.cpp
--- a/Engine/Source/Developer/GameplayDebugger/Private/GameplayDebuggerLocalController.cpp
+++ b/Engine/Source/Developer/GameplayDebugger/Private/GameplayDebuggerLocalController.cpp
@@ -167,7 +167,14 @@ void UGameplayDebuggerLocalController::OnDebugDraw(class UCanvas* Canvas, class
 			RebuildDataPackMap();
 		}
 
-		CachedReplicator->SetViewPoint(Canvas->SceneView->ViewLocation, Canvas->SceneView->ViewRotation.Vector());
+		if (Canvas->SceneView->ViewActor == nullptr)
+		{
+			CachedReplicator->SetViewPoint(Canvas->SceneView->ViewLocation, Canvas->SceneView->ViewRotation.Vector());
+		}
+		else if (CachedReplicator->IsViewPointSet())
+		{
+			CachedReplicator->ResetViewPoint();
+		}
 
 		const bool bHasDebugActor = CachedReplicator->HasDebugActor();
 		for (int32 Idx = 0; Idx < NumCategories; Idx++)
diff --git a/Engine/Source/Developer/GameplayDebugger/Public/GameplayDebuggerCategoryReplicator.h b/Engine/Source/Developer/GameplayDebugger/Public/GameplayDebuggerCategoryReplicator.h
--- a/Engine/Source/Developer/GameplayDebugger/Public/GameplayDebuggerCategoryReplicator.h
+++ b/Engine/Source/Developer/GameplayDebugger/Public/GameplayDebuggerCategoryReplicator.h
@@ -100,9 +100,15 @@ public:
 	/** [ALL] set actor for debugging */
 	void SetDebugActor(AActor* Actor, bool bSelectInEditor = false);
 
-	/** [ALL] set view matrix that should be used for culling */
+	/** [ALL] set view matrix that should be used for culling and picking instead of the associated controller view point */
 	void SetViewPoint(const FVector& InViewLocation, const FVector& InViewDirection);
 
+	/** Indicates that a view matrix has been set for culling and picking instead of the associated controller view point */
+	bool IsViewPointSet() const { return ViewLocation.IsSet() && ViewDirection.IsSet(); }
+
+	/** [ALL] reset view matrix that should be used for culling and picking to revert on the associated controller view point */
+	void ResetViewPoint();
+
 	/** [ALL] send input event to category */
 	void SendCategoryInputEvent(int32 CategoryId, int32 HandlerId);
 
@@ -218,6 +224,9 @@ protected:
 	UFUNCTION(Server, Reliable, WithValidation, meta = (CallInEditor = "true"))
 	void ServerSetViewPoint(const FVector& InViewLocation, const FVector& InViewDirection);
 
+	UFUNCTION(Server, Reliable, WithValidation, meta = (CallInEditor = "true"))
+	void ServerResetViewPoint();
+
 	UFUNCTION(Server, Reliable, WithValidation, meta = (CallInEditor = "true"))
 	void ServerSetCategoryEnabled(int32 CategoryId, bool bEnable);
 
