[GameplayDebugger] Prevent access to the GLevelEditorModeTools from commandlet environment
#rnx

diff --git a/Engine/Source/Developer/GameplayDebugger/Private/Editor/GameplayDebuggerEdMode.cpp b/Engine/Source/Developer/GameplayDebugger/Private/Editor/GameplayDebuggerEdMode.cpp
--- a/Engine/Source/Developer/GameplayDebugger/Private/Editor/GameplayDebuggerEdMode.cpp
+++ b/Engine/Source/Developer/GameplayDebugger/Private/Editor/GameplayDebuggerEdMode.cpp
@@ -111,12 +111,24 @@ void FGameplayDebuggerEdMode::Tick(FEditorViewportClient* ViewportClient, float
 
 void FGameplayDebuggerEdMode::SafeOpenMode()
 {
+	// The global mode manager should not be created or accessed in a commandlet environment
+	if (IsRunningCommandlet())
+	{
+		return;
+	}
+
 	// By calling this get, it make sure that the singleton is created and not during the SafeCloseMode whitch can be called during garbage collect.
 	GLevelEditorModeTools();
 }
 
 void FGameplayDebuggerEdMode::SafeCloseMode()
 {
+	// The global mode manager should not be created or accessed in a commandlet environment
+	if (IsRunningCommandlet())
+	{
+		return;
+	}
+
 	// this may be called on closing editor during PIE (~viewport -> teardown PIE -> debugger's cleanup on game end)
 	//
 	// DeactivateMode tries to bring up default mode, but toolkit is already destroyed by that time
