[GameplayDebugger] Hide OnScreenMessages when the gameplay debugger is active
#tests Tested in UYE
#rnx

diff --git a/Engine/Source/Developer/GameplayDebugger/Private/GameplayDebuggerLocalController.cpp b/Engine/Source/Developer/GameplayDebugger/Private/GameplayDebuggerLocalController.cpp
--- a/Engine/Source/Developer/GameplayDebugger/Private/GameplayDebuggerLocalController.cpp
+++ b/Engine/Source/Developer/GameplayDebugger/Private/GameplayDebuggerLocalController.cpp
@@ -39,6 +39,7 @@ UGameplayDebuggerLocalController::UGameplayDebuggerLocalController(const FObject
 	bIsLocallyEnabled = false;
 	bPrevLocallyEnabled = false;
 	bEnableTextShadow = false;
+	bPrevScreenMessagesEnabled = false;
 	ActiveRowIdx = 0;
 #if WITH_EDITOR
 	bActivateOnPIEEnd = false;
@@ -525,7 +526,7 @@ void UGameplayDebuggerLocalController::ToggleActivation()
 {
 	if (CachedReplicator)
 	{
-		UWorld* World = CachedReplicator->GetWorld();
+		const UWorld* World = CachedReplicator->GetWorld();
 		if (!bIsSelectingActor || StartSelectingActorHandle.IsValid())
 		{
 			bIsLocallyEnabled = !CachedReplicator->IsEnabled();
@@ -533,9 +534,19 @@ void UGameplayDebuggerLocalController::ToggleActivation()
 
 			if (bIsLocallyEnabled)
 			{
+				bPrevScreenMessagesEnabled = GAreScreenMessagesEnabled;
+				GAreScreenMessagesEnabled = false;
 				DebugActorCandidate = nullptr;
 				OnSelectActorTick();
 			}
+			else
+			{
+				// if Screen message are still disabled, restore previous state
+				if (!GAreScreenMessagesEnabled)
+				{
+					GAreScreenMessagesEnabled = bPrevScreenMessagesEnabled;
+				}
+			}
 		}
 
 		World->GetTimerManager().ClearTimer(StartSelectingActorHandle);
diff --git a/Engine/Source/Developer/GameplayDebugger/Public/GameplayDebuggerLocalController.h b/Engine/Source/Developer/GameplayDebugger/Public/GameplayDebuggerLocalController.h
--- a/Engine/Source/Developer/GameplayDebugger/Public/GameplayDebuggerLocalController.h
+++ b/Engine/Source/Developer/GameplayDebugger/Public/GameplayDebuggerLocalController.h
@@ -67,6 +67,7 @@ protected:
 	uint32 bIsLocallyEnabled : 1;
 	uint32 bPrevLocallyEnabled : 1;
 	uint32 bEnableTextShadow : 1;
+	uint32 bPrevScreenMessagesEnabled : 1;
 #if WITH_EDITOR
 	uint32 bActivateOnPIEEnd : 1;
 #endif // WITH_EDITOR
