[GameplayDebugger] visibility of the replicator actor is code driven (hidden on dedicated server, visible otherwise) and should not be replicated.
This fixes some edge cases where value was resent by server after BeginPlay resulting in component not rendering the GameplayCategories on the client.

diff --git a/Engine/Source/Developer/GameplayDebugger/Private/GameplayDebuggerCategoryReplicator.cpp b/Engine/Source/Developer/GameplayDebugger/Private/GameplayDebuggerCategoryReplicator.cpp
--- a/Engine/Source/Developer/GameplayDebugger/Private/GameplayDebuggerCategoryReplicator.cpp
+++ b/Engine/Source/Developer/GameplayDebugger/Private/GameplayDebuggerCategoryReplicator.cpp
@@ -530,6 +530,9 @@ void AGameplayDebuggerCategoryReplicator::GetLifetimeReplicatedProps(TArray<FLif
 	DOREPLIFETIME(AGameplayDebuggerCategoryReplicator, VisLogSync);
 	DOREPLIFETIME(AGameplayDebuggerCategoryReplicator, bIsEnabled);
 	DOREPLIFETIME(AGameplayDebuggerCategoryReplicator, ReplicatedData);
+
+	// The visibility of the replicator actor is code driven (hidden on dedicated server, visible otherwise) and should not be replicated
+	DISABLE_REPLICATED_PRIVATE_PROPERTY(AActor, bHidden);
 }
 
 bool AGameplayDebuggerCategoryReplicator::ServerSetEnabled_Validate(bool bEnable)
