Fixed GameplayDebuggerPlayerManager's editor-time ticking (i.e. made sure it won't get its game-world tick via the tickable object mechanics). This change addresses a rare crash related to GameplayDebuggerPlayerManager accessing UWorld in some edge cases.

diff --git a/Engine/Source/Developer/GameplayDebugger/Private/GameplayDebuggerPlayerManager.cpp b/Engine/Source/Developer/GameplayDebugger/Private/GameplayDebuggerPlayerManager.cpp
--- a/Engine/Source/Developer/GameplayDebugger/Private/GameplayDebuggerPlayerManager.cpp
+++ b/Engine/Source/Developer/GameplayDebugger/Private/GameplayDebuggerPlayerManager.cpp
@@ -33,6 +33,7 @@ AGameplayDebuggerPlayerManager::AGameplayDebuggerPlayerManager(const FObjectInit
 
 	bIsLocal = false;
 	bInitialized = false;
+	bEditorTimeTick = false;
 }
 
 void AGameplayDebuggerPlayerManager::PostInitProperties()
@@ -52,6 +53,10 @@ void AGameplayDebuggerPlayerManager::PostInitProperties()
 		{
 			UInputDelegateBinding::BindInputDelegates(GetClass(), InputComponent);
 		}
+#if WITH_EDITORONLY_DATA 
+		const UWorld* World = GetWorld();
+		bEditorTimeTick = (World != nullptr) && (World->IsEditorWorld() == true) && (World->IsGameWorld() == false);
+#endif // WITH_EDITORONLY_DATA 
 	}
 }
 
@@ -307,15 +312,5 @@ TStatId AGameplayDebuggerPlayerManager::GetStatId() const
 {
 	RETURN_QUICK_DECLARE_CYCLE_STAT(AGameplayDebuggerPlayerManager, STATGROUP_Tickables);
 }
-
-bool AGameplayDebuggerPlayerManager::IsTickable() const
-{
-#if WITH_EDITOR
-	UWorld* World = GetWorld();
-	return (World != nullptr) && (World->IsEditorWorld() == true) && (World->IsGameWorld() == false);
-#else
-	return false;
-#endif // WITH_EDITOR
-}
 // FTickableGameObject end
 
diff --git a/Engine/Source/Developer/GameplayDebugger/Public/GameplayDebuggerPlayerManager.h b/Engine/Source/Developer/GameplayDebugger/Public/GameplayDebuggerPlayerManager.h
--- a/Engine/Source/Developer/GameplayDebugger/Public/GameplayDebuggerPlayerManager.h
+++ b/Engine/Source/Developer/GameplayDebugger/Public/GameplayDebuggerPlayerManager.h
@@ -42,7 +42,8 @@ class GAMEPLAYDEBUGGER_API AGameplayDebuggerPlayerManager : public AActor, publi
 	virtual ETickableTickType GetTickableTickType() const override;
 	virtual TStatId GetStatId() const override;
 	virtual bool IsTickableInEditor() const override { return true; }
-	virtual bool IsTickable() const override;
+	/** results in this instance not being ticked in any of the game worlds */
+	virtual bool IsTickable() const override { return bEditorTimeTick; }
 	// FTickableGameObject end
 
 	virtual void PostInitProperties() override;
@@ -85,6 +86,7 @@ protected:
 	FGameplayDebuggerPlayerData EditorWorldData;
 #endif // WITH_EDITORONLY_DATA
 
+	uint8 bEditorTimeTick : 1;
 	uint32 bHasAuthority : 1;
 	uint32 bIsLocal : 1;
 	uint32 bInitialized : 1;
