[VLogger] make sure input bindings are refreshed when toggling through command lines
- also fixed screen messages not hidden when opening gameplay debugger by holding activation key

diff --git a/Engine/Source/Developer/GameplayDebugger/Private/GameplayDebuggerLocalController.cpp b/Engine/Source/Developer/GameplayDebugger/Private/GameplayDebuggerLocalController.cpp
--- a/Engine/Source/Developer/GameplayDebugger/Private/GameplayDebuggerLocalController.cpp
+++ b/Engine/Source/Developer/GameplayDebugger/Private/GameplayDebuggerLocalController.cpp
@@ -508,7 +508,6 @@ bool UGameplayDebuggerLocalController::IsKeyBound(const FName KeyName) const
 
 void UGameplayDebuggerLocalController::OnActivationPressed()
 {
-	bPrevLocallyEnabled = bIsLocallyEnabled;
 	if (CachedReplicator)
 	{
 		const float HoldTimeThr = 0.2f * (FApp::UseFixedTimeStep() ? (FApp::GetFixedDeltaTime() * 60.0f) : 1.0f);
@@ -529,6 +528,7 @@ void UGameplayDebuggerLocalController::ToggleActivation()
 		const UWorld* World = CachedReplicator->GetWorld();
 		if (!bIsSelectingActor || StartSelectingActorHandle.IsValid())
 		{
+			bPrevLocallyEnabled = bIsLocallyEnabled;
 			bIsLocallyEnabled = !CachedReplicator->IsEnabled();
 			CachedReplicator->SetEnabled(bIsLocallyEnabled);
 
@@ -650,8 +650,11 @@ void UGameplayDebuggerLocalController::OnStartSelectingActor()
 	{
 		if (!CachedReplicator->IsEnabled())
 		{
+			bPrevLocallyEnabled = bIsLocallyEnabled;
 			bIsLocallyEnabled = true;
 			CachedReplicator->SetEnabled(bIsLocallyEnabled);
+			bPrevScreenMessagesEnabled = GAreScreenMessagesEnabled;
+			GAreScreenMessagesEnabled = false;
 		}
 
 		bIsSelectingActor = true;
