Gameplay debugger reference cleanup on replay scrub and controller logout

diff --git a/Engine/Source/Developer/GameplayDebugger/Private/GameplayDebuggerPlayerManager.cpp b/Engine/Source/Developer/GameplayDebugger/Private/GameplayDebuggerPlayerManager.cpp
--- a/Engine/Source/Developer/GameplayDebugger/Private/GameplayDebuggerPlayerManager.cpp
+++ b/Engine/Source/Developer/GameplayDebugger/Private/GameplayDebuggerPlayerManager.cpp
@@ -10,7 +10,8 @@
 #include "GameplayDebuggerConfig.h"
 #include "UnrealEngine.h"
 #include "Engine/LocalPlayer.h"
-
+#include "Net/UnrealNetwork.h"
+#include "GameFramework/GameModeBase.h"
 
 AGameplayDebuggerPlayerManager::AGameplayDebuggerPlayerManager(const FObjectInitializer& ObjectInitializer) : Super(ObjectInitializer)
 {
@@ -84,6 +85,9 @@ void AGameplayDebuggerPlayerManager::BeginPlay()
 	}
 
 	PendingRegistrations.Empty();
+
+	FNetworkReplayDelegates::OnScrubTeardown.AddUObject(this, &ThisClass::OnReplayScrubTeardown);
+	FGameModeEvents::GameModeLogoutEvent.AddUObject(this, &ThisClass::OnGameModeLogout);
 }
 
 void AGameplayDebuggerPlayerManager::EndPlay(const EEndPlayReason::Type Reason)
@@ -99,6 +103,9 @@ void AGameplayDebuggerPlayerManager::EndPlay(const EEndPlayReason::Type Reason)
 			TestData.Controller = nullptr;
 		}
 	}
+
+	FNetworkReplayDelegates::OnScrubTeardown.RemoveAll(this);
+	FGameModeEvents::GameModeLogoutEvent.RemoveAll(this);
 }
 
 void AGameplayDebuggerPlayerManager::Init()
@@ -314,3 +321,38 @@ TStatId AGameplayDebuggerPlayerManager::GetStatId() const
 }
 // FTickableGameObject end
 
+void AGameplayDebuggerPlayerManager::OnGameModeLogout(AGameModeBase* GameMode, AController* Exiting)
+{
+	if (GameMode && GameMode->GetWorld() == GetWorld())
+	{
+		UWorld* World = GetWorld();
+		for (int32 Idx = PlayerData.Num() - 1; Idx >= 0; Idx--)
+		{
+			FGameplayDebuggerPlayerData& TestData = PlayerData[Idx];
+
+			if (IsValid(TestData.Replicator) && (TestData.Replicator->GetReplicationOwner() == Exiting))
+			{
+				if (IsValid(TestData.Replicator))
+				{
+					World->DestroyActor(TestData.Replicator);
+				}
+
+				if (IsValid(TestData.Controller))
+				{
+					TestData.Controller->Cleanup();
+				}
+
+				PlayerData.RemoveAt(Idx, 1, false);
+				break;
+			}
+		}
+	}
+}
+
+void AGameplayDebuggerPlayerManager::OnReplayScrubTeardown(UWorld* InWorld)
+{
+	if (GetWorld() == InWorld)
+	{
+		UpdateAuthReplicators();
+	}
+}
\ No newline at end of file
diff --git a/Engine/Source/Developer/GameplayDebugger/Public/GameplayDebuggerPlayerManager.h b/Engine/Source/Developer/GameplayDebugger/Public/GameplayDebuggerPlayerManager.h
--- a/Engine/Source/Developer/GameplayDebugger/Public/GameplayDebuggerPlayerManager.h
+++ b/Engine/Source/Developer/GameplayDebugger/Public/GameplayDebuggerPlayerManager.h
@@ -12,6 +12,7 @@ class AGameplayDebuggerCategoryReplicator;
 class APlayerController;
 class UGameplayDebuggerLocalController;
 class UInputComponent;
+class AGameModeBase;
 
 USTRUCT()
 struct FGameplayDebuggerPlayerData
@@ -51,6 +52,9 @@ class GAMEPLAYDEBUGGER_API AGameplayDebuggerPlayerManager : public AActor, publi
 protected:
 	virtual void BeginPlay() override;
 
+	void OnReplayScrubTeardown(UWorld* InWorld);
+	void OnGameModeLogout(AGameModeBase* GameMode, AController* Exiting);
+
 public:
 	virtual void EndPlay(const EEndPlayReason::Type Reason) override;
 	void Init();
diff --git a/Engine/Source/Runtime/Engine/Private/DemoNetDriver.cpp b/Engine/Source/Runtime/Engine/Private/DemoNetDriver.cpp
--- a/Engine/Source/Runtime/Engine/Private/DemoNetDriver.cpp
+++ b/Engine/Source/Runtime/Engine/Private/DemoNetDriver.cpp
@@ -3997,7 +3997,6 @@ bool UDemoNetDriver::LoadCheckpoint(const FGotoResult& GotoResult)
 		TrackedRewindActorsByGUID.Empty();
 	}
 
-#if 1
 	TSet<const AActor*> KeepAliveActors;
 
 	// Determine if an Actor has a reference to a spectator in some way.
@@ -4093,21 +4092,6 @@ bool UDemoNetDriver::LoadCheckpoint(const FGotoResult& GotoResult)
 		ServerConnection->OwningActor = nullptr;
 	}
 
-#else
-	for (int32 i = ServerConnection->OpenChannels.Num() - 1; i >= 0; i--)
-	{
-		UChannel* OpenChannel = ServerConnection->OpenChannels[i];
-		if (OpenChannel != nullptr)
-		{
-			UActorChannel* ActorChannel = Cast<UActorChannel>(OpenChannel);
-			if (ActorChannel != nullptr && ActorChannel->GetActor() != nullptr && !ActorChannel->GetActor()->IsNetStartupActor())
-			{
-				World->DestroyActor(ActorChannel->GetActor(), true);
-			}
-		}
-	}
-#endif
-
 	ReplayHelper.ExternalDataToObjectMap.Empty();
 
 	PlaybackPackets.Empty();
@@ -4131,6 +4115,8 @@ bool UDemoNetDriver::LoadCheckpoint(const FGotoResult& GotoResult)
 	ServerConnection->Close();
 	ServerConnection->CleanUp();
 
+	FNetworkReplayDelegates::OnScrubTeardown.Broadcast(World);
+
 	// Optionally collect garbage after the old actors and connection are cleaned up - there could be a lot of pending-kill objects at this point.
 	if (CVarDemoLoadCheckpointGarbageCollect.GetValueOnGameThread() != 0)
 	{
diff --git a/Engine/Source/Runtime/Engine/Private/UnrealNetwork.cpp b/Engine/Source/Runtime/Engine/Private/UnrealNetwork.cpp
--- a/Engine/Source/Runtime/Engine/Private/UnrealNetwork.cpp
+++ b/Engine/Source/Runtime/Engine/Private/UnrealNetwork.cpp
@@ -3,6 +3,7 @@
 #include "Net/UnrealNetwork.h"
 
 FPreReplayScrub                                      FNetworkReplayDelegates::OnPreScrub;
+FReplayScrubTeardown                                 FNetworkReplayDelegates::OnScrubTeardown;
 FOnWriteGameSpecificDemoHeader                       FNetworkReplayDelegates::OnWriteGameSpecificDemoHeader;
 FOnProcessGameSpecificDemoHeader                     FNetworkReplayDelegates::OnProcessGameSpecificDemoHeader;
 FOnWriteGameSpecificFrameData                        FNetworkReplayDelegates::OnWriteGameSpecificFrameData;
diff --git a/Engine/Source/Runtime/Engine/Public/Net/UnrealNetwork.h b/Engine/Source/Runtime/Engine/Public/Net/UnrealNetwork.h
--- a/Engine/Source/Runtime/Engine/Public/Net/UnrealNetwork.h
+++ b/Engine/Source/Runtime/Engine/Public/Net/UnrealNetwork.h
@@ -39,8 +39,15 @@ inline int32 MakeRelative( int32 Value, int32 Reference, int32 Max )
 	return Reference + BestSignedDifference(Value,Reference,Max);
 }
 
+UE_DEPRECATED(5.1, "No longer used.")
 DECLARE_MULTICAST_DELEGATE_OneParam(FPreActorDestroyReplayScrub, AActor*);
+
+/** Fired at the start of replay checkpoint loading */
 DECLARE_MULTICAST_DELEGATE_OneParam(FPreReplayScrub, UWorld*);
+
+/** Fired mid-checkpoint load after the previous game state is torn down, and prior to any GC */
+DECLARE_MULTICAST_DELEGATE_OneParam(FReplayScrubTeardown, UWorld*);
+
 DECLARE_MULTICAST_DELEGATE_OneParam(FOnWriteGameSpecificDemoHeader, TArray<FString>& /*GameSpecificData*/);
 DECLARE_MULTICAST_DELEGATE_TwoParams(FOnProcessGameSpecificDemoHeader, const TArray<FString>& /*GameSpecificData*/, FString& /*Error*/);
 
@@ -65,6 +72,7 @@ struct ENGINE_API FNetworkReplayDelegates
 {
 	/** global delegate called one time prior to scrubbing */
 	static FPreReplayScrub OnPreScrub;
+	static FReplayScrubTeardown OnScrubTeardown;
 
 	/** Game specific demo headers */
 	static FOnWriteGameSpecificDemoHeader OnWriteGameSpecificDemoHeader;
