PR #9570: Fix GameplayDebuggerCategory replication for standalone clients (Contributed by [REDACTED_PERSON])

diff --git a/Engine/Source/Developer/GameplayDebugger/Private/GameplayDebuggerCategoryReplicator.cpp b/Engine/Source/Developer/GameplayDebugger/Private/GameplayDebuggerCategoryReplicator.cpp
--- a/Engine/Source/Developer/GameplayDebugger/Private/GameplayDebuggerCategoryReplicator.cpp
+++ b/Engine/Source/Developer/GameplayDebugger/Private/GameplayDebuggerCategoryReplicator.cpp
@@ -234,6 +234,9 @@ bool FGameplayDebuggerNetPack::NetDeltaSerialize(FNetDeltaSerializeInfo& DeltaPa
 			uint8 ShouldUpdateShapes = bMissingOldState || (OldState->CategoryStates[Idx].ShapesRepCounter != NewState->CategoryStates[Idx].ShapesRepCounter);
 			uint8 NumDataPacks = SavedCategory.DataPacks.Num();
 
+			FName CategoryName = CategoryOb.GetCategoryName();
+			Writer << CategoryName;
+
 			Writer.WriteBit(BaseFlags);
 			Writer.WriteBit(ShouldUpdateTextLines);
 			Writer.WriteBit(ShouldUpdateShapes);
@@ -296,18 +299,25 @@ bool FGameplayDebuggerNetPack::NetDeltaSerialize(FNetDeltaSerializeInfo& DeltaPa
 		int32 CategoryCount = 0;
 		Reader << CategoryCount;
 
-		if (CategoryCount != Owner->Categories.Num())
-		{
-			UE_LOG(LogGameplayDebugReplication, Error, TEXT("Category count mismtach! received:%d expected:%d"), CategoryCount, Owner->Categories.Num());
-			Reader.SetError();
-			return false;
-		}
+		// Replicated categories still have to be serialized even if they don't exist on the client.
+		// Create a dummy category for these cases
+		FGameplayDebuggerCategory DummyCategory;
 
 		bool bHasCategoryStateChanges = false;
 		for (int32 Idx = 0; Idx < CategoryCount; Idx++)
 		{
-			FGameplayDebuggerCategory& CategoryOb = Owner->Categories[Idx].Get();
-			UE_LOG(LogGameplayDebugReplication, Verbose, TEXT("  CATEGORY[%d]:%s"), Idx, *CategoryOb.GetCategoryName().ToString());
+			FName CategoryName;
+			Reader << CategoryName;
+
+			// Does the category exist on the client?
+			// using FindLastByPredicate since that's the only Predicate-based find that returns index.
+			const int32 FoundIndex = Owner->Categories.FindLastByPredicate([CategoryName](const TSharedRef<FGameplayDebuggerCategory>& Item)
+				{
+					return (Item->GetCategoryName() == CategoryName);
+				});
+
+			FGameplayDebuggerCategory& CategoryOb = (Owner->Categories.IsValidIndex(FoundIndex)) ? Owner->Categories[FoundIndex].Get() : DummyCategory;
+			UE_LOG(LogGameplayDebugReplication, Verbose, TEXT("  CATEGORY[%d]:%s"), FoundIndex, *CategoryOb.GetCategoryName().ToString());
 
 			uint8 BaseFlags = Reader.ReadBit();
 			uint8 ShouldUpdateTextLines = Reader.ReadBit();
@@ -316,7 +326,7 @@ bool FGameplayDebuggerNetPack::NetDeltaSerialize(FNetDeltaSerializeInfo& DeltaPa
 			uint8 NumDataPacks = 0;
 			Reader << NumDataPacks;
 
-			if ((int32)NumDataPacks != CategoryOb.ReplicatedDataPacks.Num())
+			if (FoundIndex != INDEX_NONE && (int32)NumDataPacks != CategoryOb.ReplicatedDataPacks.Num())
 			{
 				UE_LOG(LogGameplayDebugReplication, Error, TEXT("Data pack count mismtach! received:%d expected:%d"), NumDataPacks, CategoryOb.ReplicatedDataPacks.Num());
 				Reader.SetError();
@@ -363,7 +373,7 @@ bool FGameplayDebuggerNetPack::NetDeltaSerialize(FNetDeltaSerializeInfo& DeltaPa
 						Reader.Serialize(DataPacket.Data.GetData(), PacketSize);
 					}
 
-					Owner->OnReceivedDataPackPacket(Idx, DataIdx, DataPacket);
+					Owner->OnReceivedDataPackPacket(FoundIndex, DataIdx, DataPacket);
 					UE_LOG(LogGameplayDebugReplication, Verbose, TEXT("  >> replicate data pack[%d] progress:%.0f%%"), DataIdx, CategoryOb.ReplicatedDataPacks[DataIdx].GetProgress() * 100.0f);
 				}
 			}
