Iris
- Added alternative replicationmode for the GameplayDebugger that works with Iris which sends DataPackPackets using RPC's rather than using a custom deltanetserialize function. The reason for this change is that the current implementation polls and modifies source state when serializing data which we do not allow when running iris.

diff --git a/Engine/Source/Developer/GameplayDebugger/Private/GameplayDebuggerCategoryReplicator.cpp b/Engine/Source/Developer/GameplayDebugger/Private/GameplayDebuggerCategoryReplicator.cpp
--- a/Engine/Source/Developer/GameplayDebugger/Private/GameplayDebuggerCategoryReplicator.cpp
+++ b/Engine/Source/Developer/GameplayDebugger/Private/GameplayDebuggerCategoryReplicator.cpp
@@ -10,6 +10,10 @@
 #include "Net/UnrealNetwork.h"
 #include "VisualLogger/VisualLogger.h"
 
+#if UE_WITH_IRIS
+#include "Net/Iris/ReplicationSystem/ReplicationSystemUtil.h"
+#endif
+
 #include UE_INLINE_GENERATED_CPP_BY_NAME(GameplayDebuggerCategoryReplicator)
 
 #if WITH_EDITOR
@@ -28,6 +32,12 @@ static TAutoConsoleVariable<int32> CVarGameplayDebuggerRepDetails(
 	TEXT("Enable or disable very verbose replication logs for gameplay debugger"),
 	ECVF_Cheat);
 
+static TAutoConsoleVariable<int32> CVarGameplayDebuggerUseDataPackRPC(
+	TEXT("GameplayDebugger.UseDataPackRPC"),
+	0,
+	TEXT("Enable or disable use of rpc's for datapack packets for gameplay debugger"),
+	ECVF_Cheat);
+
 FNotifyGameplayDebuggerOwnerChange AGameplayDebuggerCategoryReplicator::NotifyDebuggerOwnerChange;
 
 class FNetFastCategoryBaseState : public INetDeltaBaseState
@@ -96,6 +106,109 @@ public:
 	TArray<FCategoryState> CategoryStates;
 };
 
+void AGameplayDebuggerCategoryReplicator::PreReplication(IRepChangedPropertyTracker& ChangedPropertyTracker)
+{
+#if UE_WITH_IRIS
+	if (UE::Net::FReplicationSystemUtil::GetNetHandle(this).IsValid())
+	{
+		ReplicatedData.PopulateFromOwner();
+	}
+#endif
+}
+
+void AGameplayDebuggerCategoryReplicator::OnRep_ReplicatedData()
+{
+#if UE_WITH_IRIS
+	if (UE::Net::FReplicationSystemUtil::GetNetHandle(this).IsValid())
+	{
+		ReplicatedData.ApplyToOwner();
+	}
+#endif
+}
+
+#if UE_WITH_IRIS
+void FGameplayDebuggerNetPack::PopulateFromOwner()
+{
+	// The current implementation of the FGameplayDebuggerNetPack replication does not behave like most other 
+	// replicated systems so in order to make this work for Iris replication which has much stricter rules for what can be done during serialization,
+	// we cannot rely on polling being made during the call to the custom NetDeltaSerialize method so we use this explicit method to poll the data to be replicated instead.
+	
+	// Update SavedCategories for replication, DataPacks-packets are handled using RPC`s
+	if (Owner && Owner->bIsEnabled && Owner->Categories.Num() == SavedData.Num())
+	{
+		for (int32 Idx = 0; Idx < SavedData.Num(); Idx++)
+		{
+			FGameplayDebuggerCategory& CategoryOb = Owner->Categories[Idx].Get();
+			FGameplayDebuggerCategoryData& SavedCategory = SavedData[Idx];
+
+			SavedCategory.CategoryName = CategoryOb.GetCategoryName();
+			SavedCategory.bIsEnabled = CategoryOb.bIsEnabled;
+
+			const bool bTextLinesChanged = (SavedCategory.TextLines != CategoryOb.ReplicatedLines);
+			if (bTextLinesChanged)
+			{
+				SavedCategory.TextLines = CategoryOb.ReplicatedLines;
+			}
+
+			const bool bShapesChanged = (SavedCategory.Shapes != CategoryOb.ReplicatedShapes);
+			if (bShapesChanged)
+			{
+				SavedCategory.Shapes = CategoryOb.ReplicatedShapes;
+			}
+		}
+	}
+}
+
+void FGameplayDebuggerNetPack::ApplyToOwner()
+{
+	if (!ensure(Owner))
+	{
+		return;
+	}
+
+	bool bHasCategoryStateChanges = false;
+	for (const FGameplayDebuggerCategoryData& SavedDataEntry : SavedData)
+	{
+		// Does the category exist on the client?
+		// using FindLastByPredicate since that's the only Predicate-based find that returns index.
+		const FName CategoryName = SavedDataEntry.CategoryName;
+		const int32 FoundIndex = Owner->Categories.FindLastByPredicate([CategoryName](const TSharedRef<FGameplayDebuggerCategory>& Item) { return (Item->GetCategoryName() == CategoryName); });
+		if (FoundIndex == INDEX_NONE)
+		{
+			continue;
+		}
+		
+		FGameplayDebuggerCategory& CategoryOb = Owner->Categories[FoundIndex].Get();
+		UE_LOG(LogGameplayDebugReplication, Verbose, TEXT("  CATEGORY[%d]:%s"), FoundIndex, *CategoryOb.GetCategoryName().ToString());
+		
+		if (CategoryOb.bIsEnabled != SavedDataEntry.bIsEnabled)
+		{
+			bHasCategoryStateChanges = true;
+			CategoryOb.bIsEnabled = SavedDataEntry.bIsEnabled;
+		}
+
+		if (SavedDataEntry.TextLines != CategoryOb.ReplicatedLines)
+		{
+
+			UE_LOG(LogGameplayDebugReplication, Verbose, TEXT("  >> received lines"));
+			CategoryOb.ReplicatedLines = SavedDataEntry.TextLines;
+		}
+
+		if (SavedDataEntry.Shapes != CategoryOb.ReplicatedShapes)
+		{
+			UE_LOG(LogGameplayDebugReplication, Verbose, TEXT("  >> received shapes"));
+			CategoryOb.ReplicatedShapes = SavedDataEntry.Shapes;
+		}
+	}
+
+	// force scene proxy updates if categories changed state
+	if (bHasCategoryStateChanges)
+	{
+		Owner->MarkComponentsRenderStateDirty();
+	}
+}
+#endif // #if UE_WITH_IRIS
+
 bool FGameplayDebuggerNetPack::NetDeltaSerialize(FNetDeltaSerializeInfo& DeltaParms)
 {
 	if (DeltaParms.bUpdateUnmappedObjects || Owner == nullptr)
@@ -133,11 +246,12 @@ bool FGameplayDebuggerNetPack::NetDeltaSerialize(FNetDeltaSerializeInfo& DeltaPa
 			{
 				FNetFastCategoryBaseState::FCategoryState& CategoryState = NewState->CategoryStates[Idx];
 				FGameplayDebuggerCategory& CategoryOb = Owner->Categories[Idx].Get();
-				FCategoryData& SavedCategory = SavedData[Idx];
+				FGameplayDebuggerCategoryData& SavedCategory = SavedData[Idx];
 
 				const bool bMissingOldState = (OldState == nullptr) || !OldState->CategoryStates.IsValidIndex(Idx);
 				ChangedCategories[Idx] = bMissingOldState ? 1 : 0;
 
+				SavedCategory.CategoryName = CategoryOb.GetCategoryName();
 				if (SavedCategory.bIsEnabled != CategoryOb.bIsEnabled)
 				{
 					SavedCategory.bIsEnabled = CategoryOb.bIsEnabled;
@@ -160,10 +274,13 @@ bool FGameplayDebuggerNetPack::NetDeltaSerialize(FNetDeltaSerializeInfo& DeltaPa
 					ChangedCategories[Idx]++;
 				}
 
-				const int32 NumDataPacks = CategoryOb.ReplicatedDataPacks.Num();
+				// If datapack-packets are sent using RPC`s we do not need to store them in the SavedCategoryStates or the Delta baselines
+				const bool bUseDataPackRPC = Owner->bSendDataPacksUsingRPC;
+				const int32 NumDataPacks = bUseDataPackRPC ? 0 : CategoryOb.ReplicatedDataPacks.Num();
+
 				SavedCategory.DataPacks.SetNum(NumDataPacks);
 				CategoryState.DataPackStates.SetNum(NumDataPacks);
-				for (int32 DataIdx = 0; DataIdx < CategoryOb.ReplicatedDataPacks.Num(); DataIdx++)
+				for (int32 DataIdx = 0; DataIdx < NumDataPacks; DataIdx++)
 				{
 					FGameplayDebuggerDataPack& DataPack = CategoryOb.ReplicatedDataPacks[DataIdx];
 					const bool bHasOldStatePack = !bMissingOldState && OldState->CategoryStates[Idx].DataPackStates.IsValidIndex(DataIdx);
@@ -229,15 +346,14 @@ bool FGameplayDebuggerNetPack::NetDeltaSerialize(FNetDeltaSerializeInfo& DeltaPa
 		{
 			FGameplayDebuggerCategory& CategoryOb = Owner->Categories[Idx].Get();
 			const bool bMissingOldState = (OldState == nullptr) || !OldState->CategoryStates.IsValidIndex(Idx);
-			FCategoryData& SavedCategory = SavedData[Idx];
+			FGameplayDebuggerCategoryData& SavedCategory = SavedData[Idx];
 
 			uint8 BaseFlags = SavedCategory.bIsEnabled;
 			uint8 ShouldUpdateTextLines = bMissingOldState || (OldState->CategoryStates[Idx].TextLinesRepCounter != NewState->CategoryStates[Idx].TextLinesRepCounter);
 			uint8 ShouldUpdateShapes = bMissingOldState || (OldState->CategoryStates[Idx].ShapesRepCounter != NewState->CategoryStates[Idx].ShapesRepCounter);
 			uint8 NumDataPacks = SavedCategory.DataPacks.Num();
 
-			FName CategoryName = CategoryOb.GetCategoryName();
-			Writer << CategoryName;
+			Writer << SavedCategory.CategoryName;
 
 			Writer.WriteBit(BaseFlags);
 			Writer.WriteBit(ShouldUpdateTextLines);
@@ -330,7 +446,7 @@ bool FGameplayDebuggerNetPack::NetDeltaSerialize(FNetDeltaSerializeInfo& DeltaPa
 
 			if (FoundIndex != INDEX_NONE && (int32)NumDataPacks != CategoryOb.ReplicatedDataPacks.Num())
 			{
-				UE_LOG(LogGameplayDebugReplication, Error, TEXT("Data pack count mismtach! received:%d expected:%d"), NumDataPacks, CategoryOb.ReplicatedDataPacks.Num());
+				UE_LOG(LogGameplayDebugReplication, Error, TEXT("Data pack count mismatch! received:%d expected:%d"), NumDataPacks, CategoryOb.ReplicatedDataPacks.Num());
 				Reader.SetError();
 				return false;
 			}
@@ -425,6 +541,18 @@ AGameplayDebuggerCategoryReplicator::AGameplayDebuggerCategoryReplicator(const F
 	bIsEditorWorldReplicator = false;
 	bReplicates = true;
 
+#if UE_WITH_IRIS
+	if (UE::Net::ShouldUseIrisReplication())
+	{
+		// If iris is enabled, datapack-packets are always sent using RPC`s and we also requires PreReplication to called in order to populate replicated data outside of serialization.
+		SetCallPreReplication(true);
+		bSendDataPacksUsingRPC = true;
+	}
+#else
+	// We cache this as it does not make sense to be able to toggle it other than between sessions.
+	bSendDataPacksUsingRPC = CVarGameplayDebuggerUseDataPackRPC.GetValueOnAnyThread();
+#endif
+
 	ReplicatedData.Owner = this;
 }
 
@@ -441,6 +569,19 @@ void AGameplayDebuggerCategoryReplicator::BeginPlay()
 	Init();
 }
 
+#if UE_WITH_IRIS
+void AGameplayDebuggerCategoryReplicator::BeginReplication()
+{
+	bOnlyRelevantToOwner = true;
+	Super::BeginReplication();
+}
+
+const AActor* AGameplayDebuggerCategoryReplicator::GetNetOwner() const 
+{
+	return OwnerPC.Get();
+}
+#endif
+
 #if WITH_EDITOR
 void AGameplayDebuggerCategoryReplicator::InitForEditor()
 {
@@ -768,6 +909,12 @@ void AGameplayDebuggerCategoryReplicator::CollectCategoryData(bool bForce)
 					{
 						UE_LOG(LogGameplayDebugReplication, Verbose, TEXT("Category[%d].DataPack[%d] SENT, DataVersion:%d DataSize:%d SyncCounter:%d"),
 							Idx, DataPackIdx, DataPack.Header.DataVersion, DataPack.Header.DataSize, DataPack.Header.SyncCounter);
+
+						// Send the update data pack as an reliable rpc instead
+						if (bSendDataPacksUsingRPC)
+						{
+							SendDataPackPacket(CategoryOb.GetCategoryName(), DataPackIdx, DataPack);						
+						}
 					}
 				}
 			}
@@ -775,6 +922,61 @@ void AGameplayDebuggerCategoryReplicator::CollectCategoryData(bool bForce)
 	}
 }
 
+void AGameplayDebuggerCategoryReplicator::SendDataPackPacket(FName CategoryName, int32 DataPackIdx, FGameplayDebuggerDataPack& DataPack)
+{
+	FGameplayDebuggerDataPackRPCParams Params;
+
+	Params.CategoryName = CategoryName;
+	Params.DataPackIdx = DataPackIdx;
+	Params.Header = DataPack.Header;
+
+	const int32 PacketSize = FMath::Min(FGameplayDebuggerDataPack::PacketSize, DataPack.Header.DataSize - DataPack.Header.DataOffset);
+	if (PacketSize > 0)
+	{
+		Params.Data.Append(DataPack.Data.GetData() + DataPack.Header.DataOffset, PacketSize);
+	}
+
+	UE_LOG(LogGameplayDebugReplication, Verbose, TEXT("  >> replicate data pack[%d] RPC progress:%0.f%% (offset:%d packet:%d)"),
+		DataPackIdx, DataPack.Header.DataSize ? (100.0f * (DataPack.Header.DataOffset + PacketSize) / DataPack.Header.DataSize) : 100.0f,
+		DataPack.Header.DataOffset, PacketSize);
+
+	ClientDataPackPacket(Params);
+
+	// As we are sending this as a reliable RPC we "ack" it immediately
+	if (DataPack.bNeedsConfirmation && !DataPack.bReceived)
+	{
+		DataPack.OnPacketRequest(Params.Header.DataVersion, Params.Header.DataOffset);
+	}
+};
+
+void AGameplayDebuggerCategoryReplicator::ClientDataPackPacket_Implementation(const FGameplayDebuggerDataPackRPCParams& Params)
+{
+	// Does the category exist on the client?
+	// using FindLastByPredicate since that's the only Predicate-based find that returns index.
+	const FName CategoryName = Params.CategoryName;
+	const int32 FoundIndex = Categories.FindLastByPredicate([CategoryName](const TSharedRef<FGameplayDebuggerCategory>& Item) { return (Item->GetCategoryName() == CategoryName);});	
+	if (FoundIndex == INDEX_NONE)
+	{
+		return;
+	}
+
+	// receive single packet
+	FGameplayDebuggerDataPack DataPacket;
+
+	DataPacket.Header = Params.Header;
+
+	const int32 PacketSize = FMath::Min(FGameplayDebuggerDataPack::PacketSize, DataPacket.Header.DataSize - DataPacket.Header.DataOffset);
+	if (PacketSize > 0)
+	{
+		DataPacket.Data.Append(Params.Data.GetData(), PacketSize);
+	}
+	
+	OnReceivedDataPackPacket(FoundIndex, Params.DataPackIdx, DataPacket);
+
+	const FGameplayDebuggerCategory& CategoryOb = Categories[FoundIndex].Get();
+	UE_LOG(LogGameplayDebugReplication, Verbose, TEXT("  >> replicate data pack[%d] progress:%.0f%%"), Params.DataPackIdx, CategoryOb.ReplicatedDataPacks[Params.DataPackIdx].GetProgress() * 100.0f);	
+}
+
 bool AGameplayDebuggerCategoryReplicator::GetViewPoint(FVector& OutViewLocation, FVector& OutViewDirection) const
 {
 	if (IsViewPointSet())
diff --git a/Engine/Source/Developer/GameplayDebugger/Private/GameplayDebuggerModule.cpp b/Engine/Source/Developer/GameplayDebugger/Private/GameplayDebuggerModule.cpp
--- a/Engine/Source/Developer/GameplayDebugger/Private/GameplayDebuggerModule.cpp
+++ b/Engine/Source/Developer/GameplayDebugger/Private/GameplayDebuggerModule.cpp
@@ -181,14 +181,6 @@ AGameplayDebuggerPlayerManager& FGameplayDebuggerModule::GetPlayerManager(UWorld
 
 void FGameplayDebuggerModule::OnWorldInitialized(UWorld* World, const UWorld::InitializationValues IVS)
 {
-#if UE_WITH_IRIS
-	if (UE::Net::ShouldUseIrisReplication())
-	{
-		// $IRIS: $TODO: Implement custom Iris NetSerilalizer to be able to support this., https://jira.it.epicgames.com/browse/UE-158457
-		return;
-	}
-#endif // UE_WITH_IRIS
-
 	// make sure that world has valid player manager, create when it doesn't
 	if (World && (World->IsGameWorld() || World->IsEditorWorld()))
 	{
diff --git a/Engine/Source/Developer/GameplayDebugger/Private/GameplayDebuggerTypes.cpp b/Engine/Source/Developer/GameplayDebugger/Private/GameplayDebuggerTypes.cpp
--- a/Engine/Source/Developer/GameplayDebugger/Private/GameplayDebuggerTypes.cpp
+++ b/Engine/Source/Developer/GameplayDebugger/Private/GameplayDebuggerTypes.cpp
@@ -11,6 +11,8 @@
 #include "Engine/World.h"
 #include "GameFramework/PlayerController.h"
 
+#include UE_INLINE_GENERATED_CPP_BY_NAME(GameplayDebuggerTypes)
+
 DEFINE_LOG_CATEGORY(LogGameplayDebug);
 
 namespace FGameplayDebuggerUtils
diff --git a/Engine/Source/Developer/GameplayDebugger/Public/GameplayDebuggerCategoryReplicator.h b/Engine/Source/Developer/GameplayDebugger/Public/GameplayDebuggerCategoryReplicator.h
--- a/Engine/Source/Developer/GameplayDebugger/Public/GameplayDebuggerCategoryReplicator.h
+++ b/Engine/Source/Developer/GameplayDebugger/Public/GameplayDebuggerCategoryReplicator.h
@@ -15,26 +15,70 @@ class FGameplayDebuggerCategory;
 class FGameplayDebuggerExtension;
 class UGameplayDebuggerRenderingComponent;
 
+// Struct used to send the DataPackPackets as RPC`s instead of using the CustomDeltaNetSerializer.
+USTRUCT()
+struct FGameplayDebuggerDataPackRPCParams
+{
+	GENERATED_BODY()
+
+	UPROPERTY()
+	FName CategoryName;
+
+	UPROPERTY()
+	int32 DataPackIdx;
+
+	UPROPERTY()
+	FGameplayDebuggerDataPackHeader Header;
+
+	UPROPERTY()
+	TArray<uint8> Data;
+};
+
+USTRUCT()
+struct FGameplayDebuggerCategoryData
+{
+	GENERATED_BODY()
+
+	UPROPERTY()
+	FName CategoryName;
+
+	UPROPERTY()
+	TArray<FString> TextLines;
+
+	UPROPERTY()
+	TArray<FGameplayDebuggerShape> Shapes;
+
+	// Either replicated using the NetDeltaSerialize or alternatively as regular RPC`s
+	UPROPERTY(NotReplicated)
+	TArray<FGameplayDebuggerDataPackHeader> DataPacks;
+
+	UPROPERTY()
+	bool bIsEnabled;
+};
+
 USTRUCT()
 struct FGameplayDebuggerNetPack
 {
-	GENERATED_USTRUCT_BODY()
+	GENERATED_BODY()
 
-	AGameplayDebuggerCategoryReplicator* Owner;
+	UPROPERTY(NotReplicated)
+	TObjectPtr<AGameplayDebuggerCategoryReplicator> Owner;
 
 	FGameplayDebuggerNetPack() : Owner(nullptr) {}
+
 	bool NetDeltaSerialize(FNetDeltaSerializeInfo & DeltaParms);
 	void OnCategoriesChanged();
 
+#if UE_WITH_IRIS
+	// When using Iris we cannot use the same appoach as the current NetSerializeDelta as it does polling of data and also modifies data outside of the replicated state during serialization which is not allowed when using iris.
+	// Therefore we need explicit methods to populate the state from the owner and to apply it to the owner which can be invoked before and after serialization as needed.
+	void PopulateFromOwner();
+	void ApplyToOwner();
+#endif
+
 private:
-	struct FCategoryData
-	{
-		TArray<FString> TextLines;
-		TArray<FGameplayDebuggerShape> Shapes;
-		TArray<FGameplayDebuggerDataPack::FHeader> DataPacks;
-		bool bIsEnabled;
-	};
-	TArray<FCategoryData> SavedData;
+	UPROPERTY()
+	TArray<FGameplayDebuggerCategoryData> SavedData;
 };
 
 template<>
@@ -51,6 +95,7 @@ struct FGameplayDebuggerDebugActor
 {
 	GENERATED_USTRUCT_BODY()
 
+	UPROPERTY(NotReplicated)
 	TWeakObjectPtr<AActor> Actor;
 
 	UPROPERTY()
@@ -82,11 +127,19 @@ class GAMEPLAYDEBUGGER_API AGameplayDebuggerCategoryReplicator : public AActor
 protected:
 	virtual void BeginPlay() override;
 
+#if UE_WITH_IRIS
+	virtual void BeginReplication() override;
+
+	virtual const AActor* GetNetOwner() const override;
+#endif
+
 public:
 	virtual void EndPlay(const EEndPlayReason::Type EndPlayReason) override;
 	virtual void TickActor(float DeltaTime, enum ELevelTick TickType, FActorTickFunction& ThisTickFunction) override;
 	virtual void PostNetReceive() override;
 
+	virtual void PreReplication(IRepChangedPropertyTracker& ChangedPropertyTracker) override;
+
 	/** [AUTH] set new owner */
 	void SetReplicatorOwner(APlayerController* InOwnerPC);
 	static FNotifyGameplayDebuggerOwnerChange NotifyDebuggerOwnerChange;
@@ -167,6 +220,9 @@ public:
 #endif // WITH_EDITOR
 
 protected:
+	UFUNCTION()
+	void OnRep_ReplicatedData();
+
 
 	friend FGameplayDebuggerNetPack;
 
@@ -176,7 +232,7 @@ protected:
 	UPROPERTY(Replicated)
 	bool bIsEnabled;
 
-	UPROPERTY(Replicated)
+	UPROPERTY(Replicated, ReplicatedUsing=OnRep_ReplicatedData)
 	FGameplayDebuggerNetPack ReplicatedData;
 
 	UPROPERTY(Replicated)
@@ -202,6 +258,7 @@ protected:
 	uint32 bHasAuthority : 1;
 	uint32 bIsLocal : 1;
 	uint32 bIsEditorWorldReplicator : 1;
+	uint32 bSendDataPacksUsingRPC : 1;
 
 	/** notify about changes in known category set */
 	void OnCategoriesChanged();
@@ -238,9 +295,16 @@ protected:
 	UFUNCTION(Server, Reliable, WithValidation, meta = (CallInEditor = "true"))
 	void ServerSendExtensionInputEvent(int32 ExtensionId, int32 HandlerId);
 
+	/** helper function for optionally sending DataPackPackets as RPC's */
+	UFUNCTION(Client, Reliable)
+	void ClientDataPackPacket(const FGameplayDebuggerDataPackRPCParams& Params);
+	
 	/** [LOCAL] notify from CategoryData replication */
 	void OnReceivedDataPackPacket(int32 CategoryId, int32 DataPackId, const FGameplayDebuggerDataPack& DataPacket);
 
 	/** called both from BeginPlay and InitForEditor to setup instance's internal */
 	void Init();
+
+private:
+	void SendDataPackPacket(FName CategoryName, int32 DataPackIdx, FGameplayDebuggerDataPack& DataPack);
 };
diff --git a/Engine/Source/Developer/GameplayDebugger/Public/GameplayDebuggerTypes.h b/Engine/Source/Developer/GameplayDebugger/Public/GameplayDebuggerTypes.h
--- a/Engine/Source/Developer/GameplayDebugger/Public/GameplayDebuggerTypes.h
+++ b/Engine/Source/Developer/GameplayDebugger/Public/GameplayDebuggerTypes.h
@@ -6,6 +6,8 @@
 #include "Engine/EngineTypes.h"
 #include "Templates/IsArrayOrRefOfTypeByPredicate.h"
 #include "Traits/IsCharEncodingCompatibleWith.h"
+#include "UObject/ObjectMacros.h"
+#include "GameplayDebuggerTypes.generated.h"
 
 class FCanvasItem;
 class UCanvas;
@@ -157,6 +159,7 @@ namespace FGameplayDebuggerCanvasStrings
 	static FString SeparatorSpace = TEXT("  ");
 }
 
+UENUM()
 enum class EGameplayDebuggerShape : uint8
 {
 	Invalid,
@@ -171,18 +174,25 @@ enum class EGameplayDebuggerShape : uint8
 	Arrow,
 };
 
+USTRUCT()
 struct GAMEPLAYDEBUGGER_API FGameplayDebuggerShape
 {
+	GENERATED_BODY()
+
 	/** points defining shape */
+	UPROPERTY()
 	TArray<FVector> ShapeData;
 
 	/** description of shape */
+	UPROPERTY()
 	FString Description;
 
 	/** color of shape */
+	UPROPERTY()
 	FColor Color;
 
 	/** type of shape */
+	UPROPERTY()
 	EGameplayDebuggerShape Type;
 
 	FGameplayDebuggerShape() : Type(EGameplayDebuggerShape::Invalid) {}
@@ -216,42 +226,52 @@ enum class EGameplayDebuggerDataPack : uint8
 	ResetOnTick,
 };
 
-struct GAMEPLAYDEBUGGER_API FGameplayDebuggerDataPack
+USTRUCT()
+struct FGameplayDebuggerDataPackHeader
 {
-	struct FHeader
-	{
-		/** version, increased every time new Data is requested */
-		int16 DataVersion;
+	GENERATED_BODY()
+
+	/** version, increased every time new Data is requested */
+	UPROPERTY()
+	int16 DataVersion;
 
-		/** debug actor sync counter */
-		int16 SyncCounter;
+	/** debug actor sync counter */
+	UPROPERTY()
+	int16 SyncCounter;
 
-		/** size of Data array */
-		int32 DataSize;
+	/** size of Data array */
+	UPROPERTY()
+	int32 DataSize;
 
-		/** offset to currently replicated portion of data */
-		int32 DataOffset;
+	/** offset to currently replicated portion of data */
+	UPROPERTY()
+	int32 DataOffset;
 
-		/** is data compressed? */
-		uint32 bIsCompressed : 1;
+	/** is data compressed? */
+	UPROPERTY()
+	uint32 bIsCompressed : 1;
 
-		FHeader() : DataVersion(0), SyncCounter(0), DataSize(0), DataOffset(0), bIsCompressed(false) {}
+	FGameplayDebuggerDataPackHeader() : DataVersion(0), SyncCounter(0), DataSize(0), DataOffset(0), bIsCompressed(false) {}
+
+	FORCEINLINE bool Equals(const FGameplayDebuggerDataPackHeader& Other) const
+	{
+		return (DataVersion == Other.DataVersion) && (DataSize == Other.DataSize) && (DataOffset == Other.DataOffset);
+	}
 
-		FORCEINLINE bool Equals(const FHeader& Other) const
-		{
-			return (DataVersion == Other.DataVersion) && (DataSize == Other.DataSize) && (DataOffset == Other.DataOffset);
-		}
+	FORCEINLINE bool operator==(const FGameplayDebuggerDataPackHeader& Other) const
+	{
+		return Equals(Other);
+	}
 
-		FORCEINLINE bool operator==(const FHeader& Other) const
-		{
-			return Equals(Other);
-		}
+	FORCEINLINE bool operator!=(const FGameplayDebuggerDataPackHeader& Other) const
+	{
+		return !Equals(Other);
+	}
+};
 
-		FORCEINLINE bool operator!=(const FHeader& Other) const
-		{
-			return !Equals(Other);
-		}
-	};
+struct GAMEPLAYDEBUGGER_API FGameplayDebuggerDataPack
+{
+	using FHeader = FGameplayDebuggerDataPackHeader;
 
 	DECLARE_DELEGATE(FOnReset);
 	DECLARE_DELEGATE_OneParam(FOnSerialize, FArchive&);
