Fix uninitialized property errors.

diff --git a/Engine/Source/Developer/GameplayDebugger/Public/GameplayDebuggerCategoryReplicator.h b/Engine/Source/Developer/GameplayDebugger/Public/GameplayDebuggerCategoryReplicator.h
--- a/Engine/Source/Developer/GameplayDebugger/Public/GameplayDebuggerCategoryReplicator.h
+++ b/Engine/Source/Developer/GameplayDebugger/Public/GameplayDebuggerCategoryReplicator.h
@@ -25,7 +25,7 @@ struct FGameplayDebuggerDataPackRPCParams
 	FName CategoryName;
 
 	UPROPERTY()
-	int32 DataPackIdx;
+	int32 DataPackIdx = 0;
 
 	UPROPERTY()
 	FGameplayDebuggerDataPackHeader Header;
@@ -53,7 +53,7 @@ struct FGameplayDebuggerCategoryData
 	TArray<FGameplayDebuggerDataPackHeader> DataPacks;
 
 	UPROPERTY()
-	bool bIsEnabled;
+	bool bIsEnabled = false;
 };
 
 USTRUCT()
diff --git a/Engine/Source/Developer/GameplayDebugger/Public/GameplayDebuggerTypes.h b/Engine/Source/Developer/GameplayDebugger/Public/GameplayDebuggerTypes.h
--- a/Engine/Source/Developer/GameplayDebugger/Public/GameplayDebuggerTypes.h
+++ b/Engine/Source/Developer/GameplayDebugger/Public/GameplayDebuggerTypes.h
@@ -195,7 +195,7 @@ struct GAMEPLAYDEBUGGER_API FGameplayDebuggerShape
 	UPROPERTY()
 	EGameplayDebuggerShape Type;
 
-	FGameplayDebuggerShape() : Type(EGameplayDebuggerShape::Invalid) {}
+	FGameplayDebuggerShape() : Color(EForceInit::ForceInit), Type(EGameplayDebuggerShape::Invalid) {}
 
 	bool operator==(const FGameplayDebuggerShape& Other) const
 	{
