Added missing include.

diff --git a/Engine/Source/Developer/GameplayDebugger/Private/GameplayDebuggerCategoryReplicator.cpp b/Engine/Source/Developer/GameplayDebugger/Private/GameplayDebuggerCategoryReplicator.cpp
--- a/Engine/Source/Developer/GameplayDebugger/Private/GameplayDebuggerCategoryReplicator.cpp
+++ b/Engine/Source/Developer/GameplayDebugger/Private/GameplayDebuggerCategoryReplicator.cpp
@@ -11,6 +11,7 @@
 #include "VisualLogger/VisualLogger.h"
 
 #if UE_WITH_IRIS
+#include "Iris/IrisConfig.h"
 #include "Net/Iris/ReplicationSystem/ReplicationSystemUtil.h"
 #endif
 
